- Create AdvancedWidgets module
- Move RadialSlider into it for shared use
- Add redirector for VirtualCam to avoid breaking existing content
- Fix RadialSlider value not updating if not bound

diff --git a/Engine/Plugins/Experimental/VirtualCamera/Config/BaseEngine.ini b/Engine/Plugins/Experimental/VirtualCamera/Config/BaseEngine.ini
--- a/Engine/Plugins/Experimental/VirtualCamera/Config/BaseEngine.ini
+++ b/Engine/Plugins/Experimental/VirtualCamera/Config/BaseEngine.ini
@@ -1,5 +1,6 @@
 [CoreRedirects]
 +ClassRedirects=(OldName="/Script/VCamCore.VCamCoreSubsystem",NewName="/Script/VCamInput.VCamInputSubsystem")
++ClassRedirects=(OldName="/Script/VirtualCamera.RadialSlider",NewName="/Script/AdvancedWidgets.RadialSlider")
 +FunctionRedirects=(OldName="VCamComponent.ClearModifiers",NewName="VCamComponent.RemoveAllModifiers")
 +FunctionRedirects=(OldName="VCamComponent.FindModifierByName",NewName="VCamComponent.GetModifierByName")
 +FunctionRedirects=(OldName="VCamComponent.FindModifiersByClass",NewName="VCamComponent.GetModifiersByClass")
diff --git a/Engine/Plugins/Experimental/VirtualCamera/Source/VirtualCamera/Private/Components/RadialSlider.cpp b/Engine/Plugins/Experimental/VirtualCamera/Source/VirtualCamera/Private/Components/RadialSlider.cpp
deleted file mode 100644
diff --git a/Engine/Plugins/Experimental/VirtualCamera/Source/VirtualCamera/Private/VirtualCamera.cpp b/Engine/Plugins/Experimental/VirtualCamera/Source/VirtualCamera/Private/VirtualCamera.cpp
--- a/Engine/Plugins/Experimental/VirtualCamera/Source/VirtualCamera/Private/VirtualCamera.cpp
+++ b/Engine/Plugins/Experimental/VirtualCamera/Source/VirtualCamera/Private/VirtualCamera.cpp
@@ -2,11 +2,13 @@
 
 #include "VirtualCamera.h"
 
-
+#include "AdvancedWidgetsModule.h"
 #include "ConcertVirtualCamera.h"
 
+
 DEFINE_LOG_CATEGORY(LogVirtualCamera);
 
+
 const IVirtualCameraModule& IVirtualCameraModule::Get()
 {
 	static const FName ModuleName = TEXT("VirtualCamera");
@@ -29,6 +31,9 @@ private:
 	virtual void StartupModule() override
 	{
 		ConcertManager = MakeUnique<FConcertVirtualCameraManager>();
+
+		// Loads widgets (ex. RadialSlider) that are potentially referenced by assets
+		FModuleManager::Get().GetModuleChecked<FAdvancedWidgetsModule>("AdvancedWidgets");
 	}
 
 	virtual void ShutdownModule() override
diff --git a/Engine/Plugins/Experimental/VirtualCamera/Source/VirtualCamera/Private/Widgets/Input/SRadialSlider.cpp b/Engine/Plugins/Experimental/VirtualCamera/Source/VirtualCamera/Private/Widgets/Input/SRadialSlider.cpp
deleted file mode 100644
diff --git a/Engine/Plugins/Experimental/VirtualCamera/Source/VirtualCamera/Public/Components/RadialSlider.h b/Engine/Plugins/Experimental/VirtualCamera/Source/VirtualCamera/Public/Components/RadialSlider.h
deleted file mode 100644
diff --git a/Engine/Plugins/Experimental/VirtualCamera/Source/VirtualCamera/Public/Widgets/Input/SRadialSlider.h b/Engine/Plugins/Experimental/VirtualCamera/Source/VirtualCamera/Public/Widgets/Input/SRadialSlider.h
deleted file mode 100644
diff --git a/Engine/Plugins/Experimental/VirtualCamera/Source/VirtualCamera/VirtualCamera.Build.cs b/Engine/Plugins/Experimental/VirtualCamera/Source/VirtualCamera/VirtualCamera.Build.cs
--- a/Engine/Plugins/Experimental/VirtualCamera/Source/VirtualCamera/VirtualCamera.Build.cs
+++ b/Engine/Plugins/Experimental/VirtualCamera/Source/VirtualCamera/VirtualCamera.Build.cs
@@ -36,6 +36,7 @@ public class VirtualCamera : ModuleRules
 			{
 				"MediaIOCore",
 				"Slate",
+				"AdvancedWidgets"
 			}
 		);
 
diff --git a/Engine/Source/Runtime/AdvancedWidgets/AdvancedWidgets.Build.cs b/Engine/Source/Runtime/AdvancedWidgets/AdvancedWidgets.Build.cs
new file mode 100644
--- /dev/null
+++ b/Engine/Source/Runtime/AdvancedWidgets/AdvancedWidgets.Build.cs
@@ -0,0 +1,22 @@
+// Copyright Epic Games, Inc. All Rights Reserved.
+
+using UnrealBuildTool;
+
+public class AdvancedWidgets : ModuleRules
+{
+	public AdvancedWidgets(ReadOnlyTargetRules Target) : base(Target)
+	{
+		PrivateDependencyModuleNames.AddRange(
+			new string[]
+			{
+				"Core",
+				"CoreUObject",
+				"Engine",
+				"InputCore",
+				"Slate",
+				"SlateCore",
+				"UMG"
+			}
+		);
+	}
+}
diff --git a/Engine/Source/Runtime/AdvancedWidgets/Private/AdvancedWidgetsModule.cpp b/Engine/Source/Runtime/AdvancedWidgets/Private/AdvancedWidgetsModule.cpp
new file mode 100644
--- /dev/null
+++ b/Engine/Source/Runtime/AdvancedWidgets/Private/AdvancedWidgetsModule.cpp
@@ -0,0 +1,17 @@
+// Copyright Epic Games, Inc. All Rights Reserved.
+
+#include "AdvancedWidgetsModule.h"
+
+#include "Modules/ModuleManager.h"
+
+IMPLEMENT_MODULE(FAdvancedWidgetsModule, AdvancedWidgets);
+
+const FName FAdvancedWidgetsModule::AdvancedWidgetsAppIdentifier = "AdvancedWidgetsApp";
+
+void FAdvancedWidgetsModule::StartupModule()
+{
+}
+
+void FAdvancedWidgetsModule::ShutdownModule()
+{
+}
diff --git a/Engine/Source/Runtime/AdvancedWidgets/Private/Components/RadialSlider.cpp b/Engine/Source/Runtime/AdvancedWidgets/Private/Components/RadialSlider.cpp
new file mode 100644
--- /dev/null
+++ b/Engine/Source/Runtime/AdvancedWidgets/Private/Components/RadialSlider.cpp
@@ -0,0 +1,307 @@
+// Copyright Epic Games, Inc. All Rights Reserved.
+
+#include "Components/RadialSlider.h"
+#include "Slate/SRadialSlider.h"
+#include "Widgets/DeclarativeSyntaxSupport.h"
+
+#define LOCTEXT_NAMESPACE "UMG"
+
+/////////////////////////////////////////////////////
+// URadialSlider
+
+static FSliderStyle* DefaultSliderStyle = nullptr;
+
+URadialSlider::URadialSlider(const FObjectInitializer& ObjectInitializer)
+	: Super(ObjectInitializer)
+{
+	bUseCustomDefaultValue = false;
+	CustomDefaultValue = 0.0f;
+
+	//Set a default range of 0 to 1 in a linear curve.
+	SliderRange.GetRichCurve()->AddKey(0.0f, 0.0f);
+	SliderRange.GetRichCurve()->AddKey(1.0f, 1.0f);
+
+	SliderHandleStartAngle = 60.0f;
+	SliderHandleEndAngle = 300.0f;
+	AngularOffset = 0.0f;
+	SliderBarColor = FLinearColor::Gray;
+	SliderProgressColor = FLinearColor::White;
+	SliderHandleColor = FLinearColor::White;
+	StepSize = 0.01f;
+	IsFocusable = true;
+	MouseUsesStep = false;
+	RequiresControllerLock = true;
+
+	if (DefaultSliderStyle == nullptr)
+	{
+		// HACK: THIS SHOULD NOT COME FROM CORESTYLE AND SHOULD INSTEAD BE DEFINED BY ENGINE TEXTURES/PROJECT SETTINGS
+		DefaultSliderStyle = new FSliderStyle(FCoreStyle::Get().GetWidgetStyle<FSliderStyle>("Slider")); //TODO: Create bespoke SlateStyle for Radial Widget
+
+		// Unlink UMG default colors from the editor settings colors.
+		DefaultSliderStyle->UnlinkColors();
+	}
+
+	WidgetStyle = *DefaultSliderStyle;
+
+#if WITH_EDITORONLY_DATA
+	AccessibleBehavior = ESlateAccessibleBehavior::Summary;
+	bCanChildrenBeAccessible = false;
+#endif
+}
+
+TSharedRef<SWidget> URadialSlider::RebuildWidget()
+{
+	MyRadialSlider = SNew(SRadialSlider)
+		.Style(&WidgetStyle)
+		.IsFocusable(IsFocusable)
+		.OnMouseCaptureBegin(BIND_UOBJECT_DELEGATE(FSimpleDelegate, HandleOnMouseCaptureBegin))
+		.OnMouseCaptureEnd(BIND_UOBJECT_DELEGATE(FSimpleDelegate, HandleOnMouseCaptureEnd))
+		.OnControllerCaptureBegin(BIND_UOBJECT_DELEGATE(FSimpleDelegate, HandleOnControllerCaptureBegin))
+		.OnControllerCaptureEnd(BIND_UOBJECT_DELEGATE(FSimpleDelegate, HandleOnControllerCaptureEnd))
+		.OnValueChanged(BIND_UOBJECT_DELEGATE(FOnFloatValueChanged, HandleOnValueChanged));
+
+	return MyRadialSlider.ToSharedRef();
+}
+
+void URadialSlider::SynchronizeProperties()
+{
+	Super::SynchronizeProperties();
+
+	FRichCurve* OurCurve = SliderRange.GetRichCurve();
+	for (int32 i = 0; i < OurCurve->Keys.Num() - 1; i++)
+	{
+		if (i > 0)
+		{
+			const float PreviousCurveKeyTime = OurCurve->Keys[i - 1].Time;
+			if (OurCurve->Keys[i].Time < PreviousCurveKeyTime)
+			{
+				OurCurve->Keys[i].Time = PreviousCurveKeyTime;
+			}
+
+			const float PreviousCurveKeyValue = OurCurve->Keys[i - 1].Value;
+			if (OurCurve->Keys[i].Value < PreviousCurveKeyValue)
+			{
+				OurCurve->Keys[i].Value = PreviousCurveKeyValue;
+			}
+		}
+
+		const float NextCurveKeyTime = OurCurve->Keys[i + 1].Time;
+		if (OurCurve->Keys[i].Time > NextCurveKeyTime)
+		{
+			OurCurve->Keys[i].Time = NextCurveKeyTime;
+		}
+
+		const float NextCurveKeyValue = OurCurve->Keys[i + 1].Value;
+		if (OurCurve->Keys[i].Value > NextCurveKeyValue)
+		{
+			OurCurve->Keys[i].Value = NextCurveKeyValue;
+		}
+	}
+
+	TAttribute<float> ValueBinding = PROPERTY_BINDING(float, Value);
+	
+	MyRadialSlider->SetMouseUsesStep(MouseUsesStep);
+	MyRadialSlider->SetRequiresControllerLock(RequiresControllerLock);
+	MyRadialSlider->SetSliderBarColor(SliderBarColor);
+	MyRadialSlider->SetSliderProgressColor(SliderProgressColor);
+	MyRadialSlider->SetSliderHandleColor(SliderHandleColor);
+	MyRadialSlider->SetValue(ValueBinding);
+	MyRadialSlider->SetUseCustomDefaultValue(bUseCustomDefaultValue);
+	MyRadialSlider->SetCustomDefaultValue(CustomDefaultValue);
+	MyRadialSlider->SetSliderRange(SliderRange);
+	MyRadialSlider->SetValueTags(ValueTags);
+	MyRadialSlider->SetSliderHandleStartAngleAndSliderHandleEndAngle(SliderHandleStartAngle, SliderHandleEndAngle);
+	MyRadialSlider->SetAngularOffset(AngularOffset);
+	MyRadialSlider->SetLocked(Locked);
+	MyRadialSlider->SetStepSize(StepSize);
+}
+
+void URadialSlider::ReleaseSlateResources(bool bReleaseChildren)
+{
+	Super::ReleaseSlateResources(bReleaseChildren);
+
+	MyRadialSlider.Reset();
+}
+
+void URadialSlider::HandleOnValueChanged(float InValue)
+{
+	OnValueChanged.Broadcast(InValue);
+}
+
+void URadialSlider::HandleOnMouseCaptureBegin()
+{
+	OnMouseCaptureBegin.Broadcast();
+}
+
+void URadialSlider::HandleOnMouseCaptureEnd()
+{
+	OnMouseCaptureEnd.Broadcast();
+}
+
+void URadialSlider::HandleOnControllerCaptureBegin()
+{
+	OnControllerCaptureBegin.Broadcast();
+}
+
+void URadialSlider::HandleOnControllerCaptureEnd()
+{
+	OnControllerCaptureEnd.Broadcast();
+}
+
+float URadialSlider::GetValue() const
+{
+	if (MyRadialSlider.IsValid() )
+	{
+		return MyRadialSlider->GetValue();
+	}
+
+	return Value;
+}
+
+float URadialSlider::GetCustomDefaultValue() const
+{
+	if (MyRadialSlider.IsValid())
+	{
+		return MyRadialSlider->GetCustomDefaultValue();
+	}
+
+	return Value;
+}
+
+float URadialSlider::GetNormalizedSliderHandlePosition() const
+{
+	if (MyRadialSlider.IsValid())
+	{
+		return MyRadialSlider->GetNormalizedSliderHandlePosition();
+	}
+	
+	return 0.0f;
+}
+
+void URadialSlider::SetValue(float InValue)
+{
+	Value = InValue;
+	if (MyRadialSlider.IsValid() )
+	{
+		MyRadialSlider->SetValue(InValue);
+	}
+}
+
+void URadialSlider::SetCustomDefaultValue(float InValue)
+{
+	CustomDefaultValue = InValue;
+	if (MyRadialSlider.IsValid())
+	{
+		MyRadialSlider->SetCustomDefaultValue(InValue);
+	}
+}
+
+void URadialSlider::SetSliderRange(const FRuntimeFloatCurve& InSliderRange)
+{
+	SliderRange = InSliderRange;
+	if (MyRadialSlider.IsValid())
+	{
+		MyRadialSlider->SetSliderRange(InSliderRange);
+	}
+}
+
+void URadialSlider::SetSliderHandleStartAngle(float InValue)
+{
+	SliderHandleStartAngle = InValue;
+	if (MyRadialSlider.IsValid())
+	{
+		MyRadialSlider->SetSliderHandleStartAngleAndSliderHandleEndAngle(SliderHandleStartAngle, SliderHandleEndAngle);
+	}
+}
+
+void URadialSlider::SetSliderHandleEndAngle(float InValue)
+{
+	SliderHandleEndAngle = InValue;
+	if (MyRadialSlider.IsValid())
+	{
+		MyRadialSlider->SetSliderHandleStartAngleAndSliderHandleEndAngle(SliderHandleStartAngle, SliderHandleEndAngle);
+	}
+}
+
+void URadialSlider::SetAngularOffset(float InValue)
+{
+	AngularOffset = InValue;
+	if (MyRadialSlider.IsValid())
+	{
+		MyRadialSlider->SetAngularOffset(AngularOffset);
+	}
+}
+
+void URadialSlider::SetValueTags(const TArray<float>& InValueTags)
+{
+	ValueTags = InValueTags;
+	if (MyRadialSlider.IsValid())
+	{
+		MyRadialSlider->SetValueTags(ValueTags);
+	}
+}
+
+void URadialSlider::SetLocked(bool InLocked)
+{
+	Locked = InLocked;
+	if (MyRadialSlider.IsValid() )
+	{
+		MyRadialSlider->SetLocked(InLocked);
+	}
+}
+
+void URadialSlider::SetStepSize(float InValue)
+{
+	StepSize = InValue;
+	if (MyRadialSlider.IsValid())
+	{
+		MyRadialSlider->SetStepSize(InValue);
+	}
+}
+
+void URadialSlider::SetSliderHandleColor(FLinearColor InValue)
+{
+	SliderHandleColor = InValue;
+	if (MyRadialSlider.IsValid())
+	{
+		MyRadialSlider->SetSliderHandleColor(InValue);
+	}
+}
+
+void URadialSlider::SetSliderBarColor(FLinearColor InValue)
+{
+	SliderBarColor = InValue;
+	if (MyRadialSlider.IsValid())
+	{
+		MyRadialSlider->SetSliderBarColor(InValue);
+	}
+}
+
+void URadialSlider::SetSliderProgressColor(FLinearColor InValue)
+{
+	SliderProgressColor = InValue;
+	if (MyRadialSlider.IsValid())
+	{
+		MyRadialSlider->SetSliderProgressColor(InValue);
+	}
+}
+
+#if WITH_ACCESSIBILITY
+TSharedPtr<SWidget> URadialSlider::GetAccessibleWidget() const
+{
+	return MyRadialSlider;
+}
+#endif
+
+#if WITH_EDITOR
+
+const FText URadialSlider::GetPaletteCategory()
+{
+	return LOCTEXT("Common", "Common");
+}
+
+#endif
+
+/////////////////////////////////////////////////////
+
+#undef LOCTEXT_NAMESPACE
diff --git a/Engine/Source/Runtime/AdvancedWidgets/Private/Slate/SRadialSlider.cpp b/Engine/Source/Runtime/AdvancedWidgets/Private/Slate/SRadialSlider.cpp
new file mode 100644
--- /dev/null
+++ b/Engine/Source/Runtime/AdvancedWidgets/Private/Slate/SRadialSlider.cpp
@@ -0,0 +1,721 @@
+// Copyright Epic Games, Inc. All Rights Reserved.
+
+#include "Slate/SRadialSlider.h"
+
+#include "Framework/Application/SlateApplication.h"
+#include "Rendering/DrawElements.h"
+#if WITH_ACCESSIBILITY
+#include "Fonts/SlateFontInfo.h"
+#include "Fonts/FontMeasure.h"
+#include "Widgets/Accessibility/SlateAccessibleWidgets.h"
+#endif
+
+SRadialSlider::SRadialSlider()
+{
+#if WITH_ACCESSIBILITY
+	AccessibleBehavior = EAccessibleBehavior::Summary;
+	bCanChildrenBeAccessible = false;
+#endif
+}
+
+void SRadialSlider::Construct( const SRadialSlider::FArguments& InDeclaration )
+{
+	check(InDeclaration._Style);
+	Style = InDeclaration._Style;
+
+	bMouseUsesStep = InDeclaration._MouseUsesStep;
+	bRequiresControllerLock = InDeclaration._RequiresControllerLock;
+	LockedAttribute = InDeclaration._Locked;
+	StepSize = InDeclaration._StepSize;
+	ValueAttribute = InDeclaration._Value;
+	bUseCustomDefaultValue = InDeclaration._bUseCustomDefaultValue;
+	CustomDefaultValue = InDeclaration._CustomDefaultValue;
+	SliderRange = InDeclaration._SliderRange;
+	SliderHandleStartAngle = InDeclaration._SliderHandleStartAngle;
+	SliderHandleEndAngle = InDeclaration._SliderHandleEndAngle;
+	AngularOffset = InDeclaration._AngularOffset;
+	SliderBarColor = InDeclaration._SliderBarColor;
+	SliderProgressColor = InDeclaration._SliderProgressColor;
+	SliderHandleColor = InDeclaration._SliderHandleColor;
+	bIsFocusable = InDeclaration._IsFocusable;
+	OnMouseCaptureBegin = InDeclaration._OnMouseCaptureBegin;
+	OnMouseCaptureEnd = InDeclaration._OnMouseCaptureEnd;
+	OnControllerCaptureBegin = InDeclaration._OnControllerCaptureBegin;
+	OnControllerCaptureEnd = InDeclaration._OnControllerCaptureEnd;
+	OnValueChanged = InDeclaration._OnValueChanged;
+
+	bControllerInputCaptured = false;
+}
+
+int32 SRadialSlider::OnPaint( const FPaintArgs& Args, const FGeometry& AllottedGeometry, const FSlateRect& MyCullingRect, FSlateWindowElementList& OutDrawElements, int32 LayerId, const FWidgetStyle& InWidgetStyle, bool bParentEnabled ) const
+{
+	const float AllottedWidth = AllottedGeometry.GetLocalSize().X;
+	const float AllottedHeight = AllottedGeometry.GetLocalSize().Y;
+
+	FGeometry SliderGeometry = AllottedGeometry;
+
+	const bool bEnabled = ShouldBeEnabled(bParentEnabled);
+	const ESlateDrawEffect DrawEffects = bEnabled ? ESlateDrawEffect::None : ESlateDrawEffect::DisabledEffect;
+
+	// Draw radial slider bar
+	const FVector2D HandleSize = GetThumbImage()->ImageSize;
+	const FVector2D HalfHandleSize = 0.5f * HandleSize;
+	const float SliderRadius = FMath::Min(AllottedWidth, AllottedHeight) * 0.5f - HalfHandleSize.Y;
+	const FVector2D StartPoint(0.0f, SliderRadius);
+
+	const float MidPointAngle = FMath::Lerp(SliderHandleStartAngle, SliderHandleEndAngle, GetNormalizedSliderHandlePosition());
+
+	const FVector2D HandleLocation = StartPoint.GetRotated(MidPointAngle + AngularOffset);
+
+	TArray<FVector2D> SliderBarPoints;
+	TArray<FVector2D> ProgressBarPoints;
+	
+	const float NormalizedCustomDefaultValue = bUseCustomDefaultValue.Get() ? GetNormalizedValue(CustomDefaultValue.Get()) : GetMinValue();
+	const bool bIsCCW = GetNormalizedSliderHandlePosition() < NormalizedCustomDefaultValue;
+
+	const float CustomDefaultValueAngle = FMath::Lerp(SliderHandleStartAngle, SliderHandleEndAngle, NormalizedCustomDefaultValue);
+	const FVector2D CustomDefaultValueLocation = StartPoint.GetRotated(CustomDefaultValueAngle + AngularOffset);
+	   	  
+	ProgressBarPoints.Emplace(bIsCCW ? HandleLocation : CustomDefaultValueLocation);
+
+	static const int32 CircleResolution = 100;
+	for (int32 i = 0; i <= CircleResolution; i++)
+	{
+		const float CircleAnglePercent = float(i) / float(CircleResolution);
+		const float CurrentPointAngle = FMath::Lerp(SliderHandleStartAngle, SliderHandleEndAngle, CircleAnglePercent);
+
+		const FVector2D NewCirclePoint = StartPoint.GetRotated(CurrentPointAngle + AngularOffset);
+		
+		const bool bIsUniquePoint(i != CircleResolution);
+		AddSliderPointToArray(SliderBarPoints, bIsUniquePoint, NewCirclePoint);
+
+		// draw completed progress bar, given directionality
+		const bool bShouldDrawProgressbar = (
+			(bIsCCW && CircleAnglePercent >= GetNormalizedSliderHandlePosition() && CircleAnglePercent <= NormalizedCustomDefaultValue) ||
+			(!bIsCCW && CircleAnglePercent <= GetNormalizedSliderHandlePosition() && CircleAnglePercent >= NormalizedCustomDefaultValue)
+		);
+		
+		if (bShouldDrawProgressbar)
+		{
+			AddSliderPointToArray(ProgressBarPoints, bIsUniquePoint, NewCirclePoint);
+		}
+	}
+
+	ProgressBarPoints.AddUnique(bIsCCW ? CustomDefaultValueLocation : HandleLocation);
+
+	const FVector2D SliderMidPoint(AllottedGeometry.GetLocalSize() * 0.5f);
+	const FVector2D SliderDiameter(SliderRadius * 2.0f);
+	
+	// combine bar color alpha with widget alpha
+	FLinearColor SliderBarColorWithAlpha = SliderBarColor.Get().GetColor(InWidgetStyle);
+	SliderBarColorWithAlpha.A *= GetRenderOpacity();
+	FLinearColor ProgressBarColorWithAlpha = SliderProgressColor.Get().GetColor(InWidgetStyle);
+	ProgressBarColorWithAlpha.A *= GetRenderOpacity();
+
+	// draw slider bar
+	FSlateDrawElement::MakeLines
+	(
+		OutDrawElements,
+		LayerId,
+		SliderGeometry.ToPaintGeometry(SliderMidPoint, SliderDiameter),
+		SliderBarPoints,
+		DrawEffects,
+		SliderBarColorWithAlpha,
+		true,
+		Style->BarThickness
+	);
+
+	// draw completed progress bar
+	FSlateDrawElement::MakeLines
+	(
+		OutDrawElements,
+		LayerId,
+		SliderGeometry.ToPaintGeometry(SliderMidPoint, SliderDiameter),
+		ProgressBarPoints,
+		DrawEffects,
+		ProgressBarColorWithAlpha,
+		true,
+		Style->BarThickness
+	);
+
+	++LayerId;	
+	
+	TArray<float> DrawnTagValues;
+	for (int32 TagIndex = 0; TagIndex < ValueTags.Num(); TagIndex++)
+	{
+		const float CandidateValue = ValueTags[TagIndex];
+
+		if (CandidateValue < GetMinValue() || CandidateValue > GetMaxValue())
+		{
+			// Do not draw tags for values that are outside of the slider range
+			continue;
+		}
+		else if (DrawnTagValues.Find(CandidateValue) != INDEX_NONE) // Ignore duplicated tags
+		{
+			continue;
+		}
+		else
+		{
+			DrawnTagValues.Emplace(CandidateValue);
+		}
+
+		const float NormalizedCandidateTagPosition = GetNormalizedValue(CandidateValue);
+		static const float TextSpacing(5.0f);
+
+		const float TagAngle = FMath::Lerp(SliderHandleStartAngle, SliderHandleEndAngle, NormalizedCandidateTagPosition);
+
+		FVector2D LineStartLocation(FVector2D(StartPoint.X, StartPoint.Y));
+		FVector2D LineEndLocation(FVector2D(StartPoint.X, StartPoint.Y + TextSpacing));
+		LineStartLocation = LineStartLocation.GetRotated(TagAngle + AngularOffset);
+		LineEndLocation = LineEndLocation.GetRotated(TagAngle + AngularOffset);
+
+		TArray<FVector2D> TagLinePoints;
+		TagLinePoints.Emplace(LineStartLocation);
+		TagLinePoints.Emplace(LineEndLocation);
+
+		const bool bIsTagLineInProgressBarRange = (
+			(bIsCCW && NormalizedCandidateTagPosition >= GetNormalizedSliderHandlePosition() && NormalizedCandidateTagPosition <= NormalizedCustomDefaultValue) ||
+			(!bIsCCW && NormalizedCandidateTagPosition <= GetNormalizedSliderHandlePosition() && NormalizedCandidateTagPosition >= NormalizedCustomDefaultValue)
+			);
+
+		FLinearColor TagLineColor = bIsTagLineInProgressBarRange ? ProgressBarColorWithAlpha : SliderBarColorWithAlpha;
+
+		// draw tag line
+		FSlateDrawElement::MakeLines(
+			OutDrawElements,
+			LayerId,
+			SliderGeometry.ToPaintGeometry(SliderMidPoint, SliderDiameter),
+			TagLinePoints,
+			DrawEffects,
+			TagLineColor,
+			true,
+			Style->BarThickness
+		);
+
+		const FText TextToWrite = FText::FromString(FString::SanitizeFloat(CandidateValue));
+		const TSharedRef<FSlateFontMeasure> FontMeasureService = FSlateApplication::Get().GetRenderer()->GetFontMeasureService();
+		FSlateFontInfo MyFont = FCoreStyle::GetDefaultFontStyle("Regular", 10);
+		const FVector2D TextDrawSize = FontMeasureService->Measure(TextToWrite, MyFont);
+
+		FVector2D TextLocation(FVector2D(StartPoint.X, StartPoint.Y + TextSpacing + 10.0f));
+		TextLocation = TextLocation.GetRotated(TagAngle + AngularOffset);
+
+		const FVector2D TextTopLeftPoint = TextLocation + (AllottedGeometry.GetLocalSize() * 0.5f) - TextDrawSize * 0.5f;
+
+		// draw tag value
+		FSlateDrawElement::MakeText(
+			OutDrawElements,
+			LayerId,
+			SliderGeometry.ToPaintGeometry(TextTopLeftPoint, GetThumbImage()->ImageSize),
+			TextToWrite,
+			MyFont,
+			DrawEffects,
+			FLinearColor::White
+		);
+	}
+
+	// Draw slider thumb
+	const FVector2D HandleTopLeftPoint = HandleLocation + (AllottedGeometry.GetLocalSize() * 0.5f) - HalfHandleSize;
+
+	auto ThumbImage = GetThumbImage();
+	FSlateDrawElement::MakeRotatedBox(
+		OutDrawElements,
+		LayerId,
+		SliderGeometry.ToPaintGeometry(HandleTopLeftPoint, GetThumbImage()->ImageSize),
+		ThumbImage,
+		DrawEffects,
+		(180.0f + MidPointAngle + AngularOffset) * (PI / 180.0f),
+		HalfHandleSize,
+		FSlateDrawElement::RelativeToElement,
+		ThumbImage->GetTint(InWidgetStyle) * SliderHandleColor.Get().GetColor(InWidgetStyle) * InWidgetStyle.GetColorAndOpacityTint()
+	);
+
+	++LayerId;
+		
+	return LayerId;
+}
+
+void SRadialSlider::AddSliderPointToArray(TArray<FVector2D>& SliderPoints, const bool bIsUnique, const FVector2D& SliderPoint) const
+{
+	if (bIsUnique)
+	{
+		SliderPoints.AddUnique(SliderPoint);
+	}
+	else
+	{
+		SliderPoints.Emplace(SliderPoint);
+	}
+}
+
+FVector2D SRadialSlider::ComputeDesiredSize( float ) const
+{
+	return FVector2D(200.0f);
+}
+
+bool SRadialSlider::IsLocked() const
+{
+	return LockedAttribute.Get();
+}
+
+bool SRadialSlider::IsInteractable() const
+{
+	return IsEnabled() && !IsLocked() && SupportsKeyboardFocus();
+}
+
+bool SRadialSlider::SupportsKeyboardFocus() const
+{
+	return bIsFocusable;
+}
+
+void SRadialSlider::ResetControllerState()
+{
+	if (bControllerInputCaptured)
+	{
+		OnControllerCaptureEnd.ExecuteIfBound();
+		bControllerInputCaptured = false;
+	}
+}
+
+FNavigationReply SRadialSlider::OnNavigation(const FGeometry& MyGeometry, const FNavigationEvent& InNavigationEvent)
+{
+	if (bControllerInputCaptured || !bRequiresControllerLock)
+	{
+		FNavigationReply Reply = FNavigationReply::Escape();
+
+		float NewValue = ValueAttribute.Get();
+		if (InNavigationEvent.GetNavigationType() == EUINavigation::Left)
+		{
+			NewValue -= StepSize.Get();
+			Reply = FNavigationReply::Stop();
+		}
+		else if (InNavigationEvent.GetNavigationType() == EUINavigation::Right)
+		{
+			NewValue += StepSize.Get();
+			Reply = FNavigationReply::Stop();
+		}
+
+		if (ValueAttribute.Get() != NewValue)
+		{
+			CommitValue(FMath::Clamp(NewValue, GetMinValue(), GetMaxValue()));
+			return Reply;
+		}
+	}
+
+	return SLeafWidget::OnNavigation(MyGeometry, InNavigationEvent);
+}
+
+FReply SRadialSlider::OnKeyDown(const FGeometry& MyGeometry, const FKeyEvent& InKeyEvent)
+{
+	FReply Reply = FReply::Unhandled();
+
+	if (IsInteractable())
+	{
+		// The controller's bottom face button must be pressed once to begin manipulating the slider's value.
+		// Navigation away from the widget is prevented until the button has been pressed again or focus is lost.
+		// The value can be manipulated by using the game pad's directional arrows ( relative to slider orientation ).
+		if (FSlateApplication::Get().GetNavigationActionFromKey(InKeyEvent) == EUINavigationAction::Accept && bRequiresControllerLock)
+		{
+			if (bControllerInputCaptured == false)
+			{
+				// Begin capturing controller input and allow user to modify the slider's value.
+				bControllerInputCaptured = true;
+				OnControllerCaptureBegin.ExecuteIfBound();
+				Reply = FReply::Handled();
+			}
+			else
+			{
+				ResetControllerState();
+				Reply = FReply::Handled();
+			}
+		}
+		else
+		{
+			Reply = SLeafWidget::OnKeyDown(MyGeometry, InKeyEvent);
+		}
+	}
+	else
+	{
+		Reply = SLeafWidget::OnKeyDown(MyGeometry, InKeyEvent);
+	}
+
+	return Reply;
+}
+
+FReply SRadialSlider::OnKeyUp(const FGeometry& MyGeometry, const FKeyEvent& InKeyEvent)
+{
+	FReply Reply = FReply::Unhandled();
+	if (bControllerInputCaptured)
+	{
+		Reply = FReply::Handled();
+	}
+	return Reply;
+}
+
+void SRadialSlider::OnFocusLost(const FFocusEvent& InFocusEvent)
+{
+	if (bControllerInputCaptured)
+	{
+		// Commit and reset state
+		CommitValue(ValueAttribute.Get());
+		ResetControllerState();
+	}
+}
+
+void SRadialSlider::OnInputStarted(const FGeometry& MyGeometry, const FVector2D& InputAbsolutePosition)
+{
+	AbsoluteInputAngle = GetAngleFromPosition(MyGeometry, InputAbsolutePosition);
+	PreviousAbsolutePosition = InputAbsolutePosition;
+}
+
+FReply SRadialSlider::OnMouseButtonDown( const FGeometry& MyGeometry, const FPointerEvent& MouseEvent )
+{
+	if ((MouseEvent.GetEffectingButton() == EKeys::LeftMouseButton) && !IsLocked())
+	{
+		CachedCursor = GetCursor().Get(EMouseCursor::Default);
+		OnInputStarted(MyGeometry, MouseEvent.GetLastScreenSpacePosition());
+
+		OnMouseCaptureBegin.ExecuteIfBound();
+		CommitValue(PositionToValue(MyGeometry, MouseEvent.GetLastScreenSpacePosition()));
+		
+		// Release capture for controller/keyboard when switching to mouse.
+		ResetControllerState();
+		
+		return FReply::Handled().CaptureMouse(SharedThis(this));
+	}
+
+	return FReply::Unhandled();
+}
+
+FReply SRadialSlider::OnMouseButtonUp( const FGeometry& MyGeometry, const FPointerEvent& MouseEvent )
+{
+	if ((MouseEvent.GetEffectingButton() == EKeys::LeftMouseButton) && HasMouseCaptureByUser(MouseEvent.GetUserIndex(), MouseEvent.GetPointerIndex()))
+	{
+		SetCursor(CachedCursor);
+		OnMouseCaptureEnd.ExecuteIfBound();
+		
+		// Release capture for controller/keyboard when switching to mouse.
+		ResetControllerState();
+		
+		return FReply::Handled().ReleaseMouseCapture();	
+	}
+
+	return FReply::Unhandled();
+}
+
+FReply SRadialSlider::OnMouseMove( const FGeometry& MyGeometry, const FPointerEvent& MouseEvent )
+{
+	if (HasMouseCaptureByUser(MouseEvent.GetUserIndex(), MouseEvent.GetPointerIndex()) && !IsLocked())
+	{
+		SetCursor(EMouseCursor::GrabHandClosed);
+		CommitValue(PositionToValue(MyGeometry, MouseEvent.GetLastScreenSpacePosition()));
+		
+		// Release capture for controller/keyboard when switching to mouse
+		ResetControllerState();
+
+		return FReply::Handled();
+	}
+
+	return FReply::Unhandled();
+}
+
+FReply SRadialSlider::OnTouchStarted(const FGeometry& MyGeometry, const FPointerEvent& InTouchEvent)
+{
+	if (!IsLocked())
+	{
+		// Release capture for controller/keyboard when switching to mouse.
+		ResetControllerState();
+
+		PressedScreenSpaceTouchDownPosition = InTouchEvent.GetScreenSpacePosition();
+		OnInputStarted(MyGeometry, PressedScreenSpaceTouchDownPosition);
+
+		return FReply::Handled();
+	}
+
+	return FReply::Unhandled();
+}
+
+FReply SRadialSlider::OnTouchMoved(const FGeometry& MyGeometry, const FPointerEvent& InTouchEvent)
+{
+	if (HasMouseCaptureByUser(InTouchEvent.GetUserIndex(), InTouchEvent.GetPointerIndex()))
+	{
+		CommitValue(PositionToValue(MyGeometry, InTouchEvent.GetScreenSpacePosition()));
+
+		// Release capture for controller/keyboard when switching to mouse
+		ResetControllerState();
+
+		return FReply::Handled();
+	}
+	else if (!HasMouseCapture())
+	{
+		if (FSlateApplication::Get().HasTraveledFarEnoughToTriggerDrag(InTouchEvent, PressedScreenSpaceTouchDownPosition, EOrientation::Orient_Horizontal))
+		{
+			CachedCursor = GetCursor().Get(EMouseCursor::Default);
+			OnMouseCaptureBegin.ExecuteIfBound();
+
+			CommitValue(PositionToValue(MyGeometry, InTouchEvent.GetScreenSpacePosition()));
+
+			// Release capture for controller/keyboard when switching to mouse
+			ResetControllerState();
+
+			return FReply::Handled().CaptureMouse(SharedThis(this));
+		}
+	}
+
+	return FReply::Unhandled();
+}
+
+FReply SRadialSlider::OnTouchEnded(const FGeometry& MyGeometry, const FPointerEvent& InTouchEvent)
+{
+	if (HasMouseCaptureByUser(InTouchEvent.GetUserIndex(), InTouchEvent.GetPointerIndex()))
+	{
+		SetCursor(CachedCursor);
+		OnMouseCaptureEnd.ExecuteIfBound();
+
+		CommitValue(PositionToValue(MyGeometry, InTouchEvent.GetScreenSpacePosition()));
+
+		// Release capture for controller/keyboard when switching to mouse.
+		ResetControllerState();
+
+		return FReply::Handled().ReleaseMouseCapture();
+	}
+
+	return FReply::Unhandled();
+}
+
+void SRadialSlider::CommitValue(float NewValue)
+{
+	ValueAttribute.Set(NewValue);
+
+	Invalidate(EInvalidateWidgetReason::Paint);
+
+	OnValueChanged.ExecuteIfBound(FMath::Clamp(NewValue, GetMinValue(), GetMaxValue()));
+}
+
+float SRadialSlider::GetAngleFromPosition(const FGeometry& MyGeometry, const FVector2D& AbsolutePosition)
+{
+	const FVector2D LocalPosition = MyGeometry.AbsoluteToLocal(AbsolutePosition) - MyGeometry.GetLocalSize() * 0.5f;
+	const FVector2D InputDirection = LocalPosition.GetSafeNormal().GetRotated(90.0f - AngularOffset);
+	const float Angle = 180.0f + (180.0f / PI * FMath::Atan2(InputDirection.Y, InputDirection.X));
+
+	return Angle;
+}
+
+float SRadialSlider::PositionToValue(const FGeometry& MyGeometry, const FVector2D& AbsolutePosition)
+{
+	const float OldAngle = GetAngleFromPosition(MyGeometry, PreviousAbsolutePosition);
+	const float NewAngle = GetAngleFromPosition(MyGeometry, AbsolutePosition);
+	
+	float WrappedDelta = FMath::Abs(FMath::Fmod((NewAngle - OldAngle), 360.0f));
+	float AngularDelta = WrappedDelta > 180.0f ? 360.0f - WrappedDelta : WrappedDelta;
+
+	int32 Sign = (NewAngle - OldAngle >= 0.0f && NewAngle - OldAngle <= 180.0f) || (NewAngle - OldAngle <= -180 && NewAngle - OldAngle >= -360) ? 1 : -1;
+	AngularDelta *= Sign;
+
+	PreviousAbsolutePosition = AbsolutePosition;
+	AbsoluteInputAngle += AngularDelta;
+
+	const FRichCurve* SliderRangeCurve = SliderRange.GetRichCurveConst();
+	
+	float CurveMinTime;
+	float CurveMaxTime;
+	SliderRangeCurve->GetTimeRange(CurveMinTime, CurveMaxTime);
+
+	float CurveTime = FMath::GetMappedRangeValueClamped(FVector2D(SliderHandleStartAngle, SliderHandleEndAngle), FVector2D(CurveMinTime, CurveMaxTime), AbsoluteInputAngle);
+	float NewValue = SliderRangeCurve->Eval(CurveTime);
+
+	if (bMouseUsesStep)
+	{
+		const float SteppedValue = FMath::RoundToInt(NewValue / StepSize.Get()) * StepSize.Get();
+		NewValue = SteppedValue;
+	}
+
+	return NewValue;
+}
+
+const FSlateBrush* SRadialSlider::GetBarImage() const
+{
+	if (!IsEnabled() || LockedAttribute.Get())
+	{
+		return &Style->DisabledBarImage;
+	}
+	else if (IsHovered())
+	{
+		return &Style->HoveredBarImage;
+	}
+	else
+	{
+		return &Style->NormalBarImage;
+	}
+}
+
+const FSlateBrush* SRadialSlider::GetThumbImage() const
+{
+	if (!IsEnabled() || LockedAttribute.Get())
+	{
+		return &Style->DisabledThumbImage;
+	}
+	else if (IsHovered())
+	{
+		return &Style->HoveredThumbImage;
+	}
+	else
+	{
+		return &Style->NormalThumbImage;
+	}
+}
+
+float SRadialSlider::GetMinValue() const
+{
+	float MinValue = 0.0f;
+
+	const FRichCurve* SliderRangeCurve = SliderRange.GetRichCurveConst();
+
+	if (SliderRangeCurve->GetNumKeys() > 0)
+	{
+		MinValue = SliderRangeCurve->GetFirstKey().Value;
+	}
+	
+	return MinValue;
+}
+
+float SRadialSlider::GetMaxValue() const
+{
+	float MaxValue = 0.0f;
+
+	const FRichCurve* SliderRangeCurve = SliderRange.GetRichCurveConst();
+
+	if (SliderRangeCurve->GetNumKeys() > 0)
+	{
+		MaxValue = SliderRangeCurve->GetLastKey().Value;
+	}
+
+	return MaxValue;
+}
+
+float SRadialSlider::GetValue() const
+{
+	return ValueAttribute.Get();
+}
+
+bool SRadialSlider::GetUseCustomDefaultValue() const
+{
+	return bUseCustomDefaultValue.Get();
+}
+
+float SRadialSlider::GetCustomDefaultValue() const
+{
+	return CustomDefaultValue.Get();
+}
+
+float SRadialSlider::GetNormalizedValue(float RawValue) const
+{
+	float NormalizedValue(0.0f);	
+
+	const TArray<FRichCurveKey>& SliderRangeCurveKeys = SliderRange.GetRichCurveConst()->GetConstRefOfKeys();
+	
+	if (SliderRangeCurveKeys.Num() > 0)
+	{
+		float SliderPercentMin(0.0f);
+		float SliderPercentMax(1.0f);
+
+		float SampleValueMin(GetMinValue());
+		float SampleValueMax(GetMaxValue());
+
+		const float MinCurveTime = SliderRangeCurveKeys[0].Time;
+		const float MaxCurveTime = SliderRangeCurveKeys[SliderRangeCurveKeys.Num() -1].Time;
+
+		for (int32 i = 0; i < SliderRangeCurveKeys.Num(); i++)
+		{
+			if (RawValue < SliderRangeCurveKeys[i].Value)
+			{
+				SliderPercentMax = (SliderRangeCurveKeys[i].Time - MinCurveTime) / (MaxCurveTime - MinCurveTime);
+				SampleValueMax = SliderRangeCurveKeys[i].Value;
+				break;
+			}
+
+			SliderPercentMin = (SliderRangeCurveKeys[i].Time - MinCurveTime) / (MaxCurveTime - MinCurveTime);
+			SampleValueMin = SliderRangeCurveKeys[i].Value;
+		}
+
+		NormalizedValue = FMath::GetMappedRangeValueClamped(FVector2D(SampleValueMin, SampleValueMax), FVector2D(SliderPercentMin, SliderPercentMax), RawValue);
+	}
+
+	return NormalizedValue;
+}
+
+float SRadialSlider::GetNormalizedSliderHandlePosition() const
+{
+	return GetNormalizedValue(ValueAttribute.Get());
+}
+
+void SRadialSlider::SetValue(const TAttribute<float>& InValueAttribute)
+{
+	SetAttribute(ValueAttribute, InValueAttribute, EInvalidateWidgetReason::Paint);
+}
+
+void SRadialSlider::SetUseCustomDefaultValue(const TAttribute<bool>& InValueAttribute)
+{
+	SetAttribute(bUseCustomDefaultValue, InValueAttribute, EInvalidateWidgetReason::Paint);
+}
+
+void SRadialSlider::SetCustomDefaultValue(const TAttribute<float>& InValueAttribute)
+{
+	SetAttribute(CustomDefaultValue, InValueAttribute, EInvalidateWidgetReason::Paint);
+}
+
+void SRadialSlider::SetSliderHandleStartAngleAndSliderHandleEndAngle(float InSliderHandleStartAngle, float InSliderHandleEndAngle)
+{
+	SliderHandleStartAngle = InSliderHandleStartAngle;
+	SliderHandleEndAngle = InSliderHandleEndAngle;
+	if (SliderHandleStartAngle > SliderHandleEndAngle)
+	{
+		SliderHandleEndAngle = SliderHandleStartAngle;
+	}
+}
+
+void SRadialSlider::SetLocked(const TAttribute<bool>& InLocked)
+{
+	SetAttribute(LockedAttribute, InLocked, EInvalidateWidgetReason::Paint);
+}
+
+void SRadialSlider::SetSliderBarColor(FSlateColor InSliderBarColor)
+{
+	SetAttribute(SliderBarColor, TAttribute<FSlateColor>(InSliderBarColor), EInvalidateWidgetReason::Paint);
+}
+
+void SRadialSlider::SetSliderProgressColor(FSlateColor InSliderProgressColor)
+{
+	SetAttribute(SliderProgressColor, TAttribute<FSlateColor>(InSliderProgressColor), EInvalidateWidgetReason::Paint);
+}
+
+void SRadialSlider::SetSliderHandleColor(FSlateColor InSliderHandleColor)
+{
+	SetAttribute(SliderHandleColor, TAttribute<FSlateColor>(InSliderHandleColor), EInvalidateWidgetReason::Paint);
+}
+
+float SRadialSlider::GetStepSize() const
+{
+	return StepSize.Get();
+}
+
+void SRadialSlider::SetStepSize(const TAttribute<float>& InStepSize)
+{
+	StepSize = InStepSize;
+}
+
+void SRadialSlider::SetMouseUsesStep(bool MouseUsesStep)
+{
+	bMouseUsesStep = MouseUsesStep;
+}
+
+void SRadialSlider::SetRequiresControllerLock(bool RequiresControllerLock)
+{
+	bRequiresControllerLock = RequiresControllerLock;
+}
+
+#if WITH_ACCESSIBILITY
+TSharedRef<FSlateAccessibleWidget> SRadialSlider::CreateAccessibleWidget()
+{
+	return MakeShareable<FSlateAccessibleWidget>(new FSlateAccessibleSlider(SharedThis(this)));
+}
+#endif
diff --git a/Engine/Source/Runtime/AdvancedWidgets/Public/AdvancedWidgets.h b/Engine/Source/Runtime/AdvancedWidgets/Public/AdvancedWidgets.h
new file mode 100644
--- /dev/null
+++ b/Engine/Source/Runtime/AdvancedWidgets/Public/AdvancedWidgets.h
@@ -0,0 +1,11 @@
+// Copyright Epic Games, Inc. All Rights Reserved.
+#pragma once
+
+#include "Misc/MonolithicHeaderBoilerplate.h"
+MONOLITHIC_HEADER_BOILERPLATE()
+
+#include "Core.h"
+#include "Engine.h"
+#include "SlateBasics.h"
+
+#include "AdvancedWidgetsModule.h"
diff --git a/Engine/Source/Runtime/AdvancedWidgets/Public/AdvancedWidgetsModule.h b/Engine/Source/Runtime/AdvancedWidgets/Public/AdvancedWidgetsModule.h
new file mode 100644
--- /dev/null
+++ b/Engine/Source/Runtime/AdvancedWidgets/Public/AdvancedWidgetsModule.h
@@ -0,0 +1,28 @@
+// Copyright Epic Games, Inc. All Rights Reserved.
+
+#pragma once
+
+#include "CoreMinimal.h"
+#include "Modules/ModuleInterface.h"
+
+
+/**
+ * Advanced Widgets module
+ */
+class ADVANCEDWIDGETS_API FAdvancedWidgetsModule : public IModuleInterface
+{
+public:
+
+	/**
+	 * Called right after the plugin DLL has been loaded and the plugin object has been created
+	 */
+	virtual void StartupModule();
+
+	/**
+	 * Called before the plugin is unloaded, right before the plugin object is destroyed.
+	 */
+	virtual void ShutdownModule();
+
+	/** Advanced Widgets app identifier string */
+	static const FName AdvancedWidgetsAppIdentifier;
+};
diff --git a/Engine/Source/Runtime/AdvancedWidgets/Public/Components/RadialSlider.h b/Engine/Source/Runtime/AdvancedWidgets/Public/Components/RadialSlider.h
new file mode 100644
--- /dev/null
+++ b/Engine/Source/Runtime/AdvancedWidgets/Public/Components/RadialSlider.h
@@ -0,0 +1,214 @@
+// Copyright Epic Games, Inc. All Rights Reserved.
+#pragma once
+
+#include "CoreMinimal.h"
+
+#include "Components/Slider.h"
+#include "Components/Widget.h"
+#include "Styling/SlateTypes.h"
+#include "UObject/ObjectMacros.h"
+#include "Widgets/SWidget.h"
+
+#include "RadialSlider.generated.h"
+
+
+// Forward Declarations
+class SRadialSlider;
+
+
+/**
+ * A simple widget that shows a sliding bar with a handle that allows you to control the value between 0..1.
+ *
+ * * No Children
+ */
+
+UCLASS()
+class ADVANCEDWIDGETS_API URadialSlider : public UWidget
+{
+	GENERATED_UCLASS_BODY()
+
+public:
+	/** The slider value to display. */
+	UPROPERTY(EditAnywhere, BlueprintReadWrite, Category = Appearance, BlueprintSetter = SetValue,  meta = (UIMin="0", UIMax="1"))
+	float Value;
+
+	/** A bindable delegate to allow logic to drive the value of the widget */
+	UPROPERTY()
+	FGetFloat ValueDelegate;
+	
+	/** Whether the slider should draw it's progress bar from a custom value on the slider */
+	UPROPERTY(EditAnywhere, Category = Appearance)
+	bool bUseCustomDefaultValue;
+
+	/**  The value where the slider should draw it's progress bar from, independent of direction */
+	UPROPERTY(EditAnywhere, Category = Appearance, meta = (EditCondition = "bUseCustomDefaultValue", UIMin = "0", UIMax = "1"))
+	float CustomDefaultValue;
+
+	/** A curve that defines how the slider should be sampled. Default is linear. */
+	UPROPERTY(EditAnywhere, BlueprintReadOnly, Category = Appearance)
+	FRuntimeFloatCurve SliderRange;
+
+	/** Adds text tags to the radial slider at the value's position. */
+	UPROPERTY(EditAnywhere, BlueprintReadOnly, Category = Appearance)
+	TArray<float> ValueTags;
+
+	/** The angle at which the Slider Handle will start. */
+	UPROPERTY(EditAnywhere, BlueprintReadOnly, meta=(ClampMin = "0", ClampMax = "360"), Category = Appearance)
+	float SliderHandleStartAngle;
+
+	/** The angle at which the Slider Handle will end. */
+	UPROPERTY(EditAnywhere, BlueprintReadOnly, meta=(ClampMin = "0", ClampMax = "360"), Category = Appearance)
+	float SliderHandleEndAngle;
+
+	/** Rotates radial slider by arbitrary offset to support full gamut of configurations. */
+	UPROPERTY(EditAnywhere, BlueprintReadOnly, meta=(ClampMin = "0", ClampMax = "360"), Category = Appearance)
+	float AngularOffset;
+
+	/** The progress bar style */
+	UPROPERTY(EditAnywhere, BlueprintReadWrite, Category="Style", meta=( DisplayName="Style" ))
+	FSliderStyle WidgetStyle;
+
+	/** The color to draw the slider bar in. */
+	UPROPERTY(EditAnywhere, BlueprintReadOnly, Category=Appearance)
+	FLinearColor SliderBarColor;
+
+	/** The color to draw the completed progress of the slider bar in. */
+	UPROPERTY(EditAnywhere, BlueprintReadOnly, Category = Appearance)
+	FLinearColor SliderProgressColor;
+
+	/** The color to draw the slider handle in. */
+	UPROPERTY(EditAnywhere, BlueprintReadOnly, Category=Appearance)
+	FLinearColor SliderHandleColor;
+
+	/** Whether the handle is interactive or fixed. */
+	UPROPERTY(EditAnywhere, BlueprintReadOnly, Category=Appearance, AdvancedDisplay)
+	bool Locked;
+
+	/** Sets new value if mouse position is greater/less than half the step size. */
+	UPROPERTY(EditAnywhere, BlueprintReadOnly, Category = Appearance, AdvancedDisplay)
+	bool MouseUsesStep;
+
+	/** Sets whether we have to lock input to change the slider value. */
+	UPROPERTY(EditAnywhere, BlueprintReadOnly, Category = Appearance, AdvancedDisplay)
+	bool RequiresControllerLock;
+
+	/** The amount to adjust the value by, when using a controller or keyboard */
+	UPROPERTY(EditAnywhere, BlueprintReadOnly, Category=Appearance, meta=(UIMin="0", UIMax="1"))
+	float StepSize;
+
+	/** Should the slider be focusable? */
+	UPROPERTY(EditAnywhere, BlueprintReadOnly, Category="Interaction")
+	bool IsFocusable;
+
+	/** Invoked when the mouse is pressed and a capture begins. */
+	UPROPERTY(BlueprintAssignable, Category="Widget Event")
+	FOnMouseCaptureBeginEvent OnMouseCaptureBegin;
+
+	/** Invoked when the mouse is released and a capture ends. */
+	UPROPERTY(BlueprintAssignable, Category="Widget Event")
+	FOnMouseCaptureEndEvent OnMouseCaptureEnd;
+
+	/** Invoked when the controller capture begins. */
+	UPROPERTY(BlueprintAssignable, Category = "Widget Event")
+	FOnControllerCaptureBeginEvent OnControllerCaptureBegin;
+
+	/** Invoked when the controller capture ends. */
+	UPROPERTY(BlueprintAssignable, Category = "Widget Event")
+	FOnControllerCaptureEndEvent OnControllerCaptureEnd;
+
+	/** Called when the value is changed by slider or typing. */
+	UPROPERTY(BlueprintAssignable, Category="Widget Event")
+	FOnFloatValueChangedEvent OnValueChanged;
+
+	/** Gets the current value of the slider. */
+	UFUNCTION(BlueprintCallable, Category="Behavior")
+	float GetValue() const;
+
+	/** Gets the current custom default value of the slider. */
+	UFUNCTION(BlueprintCallable, Category = "Behavior")
+	float GetCustomDefaultValue() const;
+
+	/** Get the current raw slider alpha from 0 to 1 */
+	UFUNCTION(BlueprintCallable, Category = "Behavior")
+	float GetNormalizedSliderHandlePosition() const;
+
+	/** Sets the current value of the slider. */
+	UFUNCTION(BlueprintCallable, Category="Behavior")
+	void SetValue(float InValue);
+
+	/** Sets the current custom default value of the slider. */
+	UFUNCTION(BlueprintCallable, Category = "Behavior")
+	void SetCustomDefaultValue(float InValue);
+
+	/** Sets the curve for the slider range*/
+	UFUNCTION(BlueprintCallable, Category = "Behaviour")
+	void SetSliderRange(const FRuntimeFloatCurve& InSliderRange);
+
+	/** Adds value tags to the slider. */
+	UFUNCTION(BlueprintCallable, Category = "Behaviour")
+	void SetValueTags(const TArray<float>& InValueTags);
+
+	/** Sets the minimum angle of the slider. */
+	UFUNCTION(BlueprintCallable, Category = "Behavior")
+	void SetSliderHandleStartAngle(float InValue);
+
+	/** Sets the maximum angle of the slider. */
+	UFUNCTION(BlueprintCallable, Category = "Behavior")
+	void SetSliderHandleEndAngle(float InValue);
+
+	/** Sets the Angular Offset for the slider. */
+	UFUNCTION(BlueprintCallable, Category = "Behaviour")
+	void SetAngularOffset(float InValue);
+
+	/** Sets the handle to be interactive or fixed */
+	UFUNCTION(BlueprintCallable, Category="Behavior")
+	void SetLocked(bool InValue);
+
+	/** Sets the amount to adjust the value by, when using a controller or keyboard */
+	UFUNCTION(BlueprintCallable, Category="Behavior")
+	void SetStepSize(float InValue);
+
+	/** Sets the color of the slider bar */
+	UFUNCTION(BlueprintCallable, Category="Appearance")
+	void SetSliderBarColor(FLinearColor InValue);
+
+	/** Sets the progress color of the slider bar */
+	UFUNCTION(BlueprintCallable, Category = "Appearance")
+	void SetSliderProgressColor(FLinearColor InValue);
+
+	/** Sets the color of the handle bar */
+	UFUNCTION(BlueprintCallable, Category="Appearance")
+	void SetSliderHandleColor(FLinearColor InValue);
+	
+	// UWidget interface
+	virtual void SynchronizeProperties() override;
+	// End of UWidget interface
+
+	// UVisual interface
+	virtual void ReleaseSlateResources(bool bReleaseChildren) override;
+	// End of UVisual interface
+
+#if WITH_EDITOR
+	virtual const FText GetPaletteCategory() override;
+#endif
+
+protected:
+	/** Native Slate Widget */
+	TSharedPtr<SRadialSlider> MyRadialSlider;
+
+	// UWidget interface
+	virtual TSharedRef<SWidget> RebuildWidget() override;
+	// End of UWidget interface
+
+	void HandleOnValueChanged(float InValue);
+	void HandleOnMouseCaptureBegin();
+	void HandleOnMouseCaptureEnd();
+	void HandleOnControllerCaptureBegin();
+	void HandleOnControllerCaptureEnd();
+
+#if WITH_ACCESSIBILITY
+	virtual TSharedPtr<SWidget> GetAccessibleWidget() const override;
+#endif
+
+	PROPERTY_BINDING_IMPLEMENTATION(float, Value);
+};
diff --git a/Engine/Source/Runtime/AdvancedWidgets/Public/Slate/SRadialSlider.h b/Engine/Source/Runtime/AdvancedWidgets/Public/Slate/SRadialSlider.h
new file mode 100644
--- /dev/null
+++ b/Engine/Source/Runtime/AdvancedWidgets/Public/Slate/SRadialSlider.h
@@ -0,0 +1,352 @@
+// Copyright Epic Games, Inc. All Rights Reserved.
+
+#pragma once
+
+#include "CoreMinimal.h"
+#include "Curves/CurveFloat.h"
+#include "Framework/SlateDelegates.h"
+#include "Input/Reply.h"
+#include "Misc/Attribute.h"
+#include "Styling/CoreStyle.h"
+#include "Styling/SlateBrush.h"
+#include "Styling/SlateColor.h"
+#include "Styling/SlateTypes.h"
+#include "Styling/SlateWidgetStyle.h"
+#include "Styling/SlateWidgetStyleAsset.h"
+#include "Widgets/DeclarativeSyntaxSupport.h"
+#include "Widgets/SLeafWidget.h"
+
+class FPaintArgs;
+class FSlateWindowElementList;
+class UCurveFloat;
+
+class ADVANCEDWIDGETS_API SRadialSlider
+	: public SLeafWidget
+{
+public:
+
+	SLATE_BEGIN_ARGS(SRadialSlider)
+		: _MouseUsesStep(false)
+		, _RequiresControllerLock(true)
+		, _Locked(false)
+		, _SliderBarColor(FLinearColor::Gray)
+		, _SliderProgressColor(FLinearColor::White)
+		, _SliderHandleColor(FLinearColor::White)
+		, _Style(&FCoreStyle::Get().GetWidgetStyle<FSliderStyle>("Slider")) //TODO: Create bespoke SlateStyle for Radial Widget
+		, _StepSize(0.01f)
+		, _Value(1.f)
+		, _bUseCustomDefaultValue(false)
+		, _CustomDefaultValue(0.0f)
+		, _SliderHandleStartAngle(60.0f)
+		, _SliderHandleEndAngle(300.0f)
+		, _AngularOffset(0.0f)
+		, _IsFocusable(true)
+		, _OnMouseCaptureBegin()
+		, _OnMouseCaptureEnd()
+		, _OnValueChanged()
+		{
+		}
+
+		/** Sets new value if mouse position is greater/less than half the step size. */
+		SLATE_ARGUMENT( bool, MouseUsesStep )
+
+		/** Sets whether we have to lock input to change the slider value. */
+		SLATE_ARGUMENT( bool, RequiresControllerLock )
+
+		/** Whether the handle is interactive or fixed. */
+		SLATE_ATTRIBUTE( bool, Locked )
+
+		/** The color to draw the slider bar in. */
+		SLATE_ATTRIBUTE( FSlateColor, SliderBarColor )
+
+		/** The color to draw completed progress of the slider bar in. */
+		SLATE_ATTRIBUTE( FSlateColor, SliderProgressColor )
+
+		/** The color to draw the slider handle in. */
+		SLATE_ATTRIBUTE( FSlateColor, SliderHandleColor )
+
+		/** The style used to draw the slider. */
+		SLATE_STYLE_ARGUMENT(FSliderStyle, Style )
+
+		/** The input mode while using the controller. */
+		SLATE_ATTRIBUTE(float, StepSize)
+
+		/** A value that drives where the slider handle appears. Value is normalized between 0 and 1. */
+		SLATE_ATTRIBUTE( float, Value )
+
+		/** Whether the slider should draw it's progress bar from a custom value on the slider */
+		SLATE_ATTRIBUTE(bool, bUseCustomDefaultValue)
+
+		/** The value where the slider should draw it's progress bar from, independent of direction */
+		SLATE_ATTRIBUTE(float, CustomDefaultValue)
+
+		/** A curve that defines how the slider should be sampled. Default is linear.*/
+		SLATE_ARGUMENT(FRuntimeFloatCurve, SliderRange)
+
+		/** The angle at which the Slider Handle will start. */
+		SLATE_ARGUMENT(float, SliderHandleStartAngle)
+
+		/** The angle at which the Slider Handle will end. */
+		SLATE_ARGUMENT(float, SliderHandleEndAngle)
+
+		/** Rotates radial slider by arbitrary offset to support full gamut of configurations */
+		SLATE_ARGUMENT(float, AngularOffset)
+
+		/** Distributes value tags along the slider */
+		SLATE_ARGUMENT(TArray<float>, ValueTags)
+
+		/** Sometimes a slider should only be mouse-clickable and never keyboard focusable. */
+		SLATE_ARGUMENT(bool, IsFocusable)
+
+		/** Invoked when the mouse is pressed and a capture begins. */
+		SLATE_EVENT(FSimpleDelegate, OnMouseCaptureBegin)
+
+		/** Invoked when the mouse is released and a capture ends. */
+		SLATE_EVENT(FSimpleDelegate, OnMouseCaptureEnd)
+
+		/** Invoked when the Controller is pressed and capture begins. */
+		SLATE_EVENT(FSimpleDelegate, OnControllerCaptureBegin)
+
+		/** Invoked when the controller capture is released.  */
+		SLATE_EVENT(FSimpleDelegate, OnControllerCaptureEnd)
+
+		/** Called when the value is changed by the slider. */
+		SLATE_EVENT( FOnFloatValueChanged, OnValueChanged )
+
+	SLATE_END_ARGS()
+
+	SRadialSlider();
+
+	/**
+	 * Construct the widget.
+	 * 
+	 * @param InDeclaration A declaration from which to construct the widget.
+	 */
+	void Construct( const SRadialSlider::FArguments& InDeclaration );
+
+	/** Get the SliderRange attribute */
+	FRuntimeFloatCurve GetSliderRange() const { return SliderRange; }
+
+	/** Get the minumum value in Slider Range */
+	float GetMinValue() const;
+
+	/** Get the maximum value in Slider Range */
+	float GetMaxValue() const;
+
+	/** Get the MinSliderHandleAngle attribute */
+	float GetSliderHandleStartAngle() const { return SliderHandleStartAngle; }
+
+	/** Get the MaxSliderHandleAngle attribute */
+	float GetSliderHandleEndAngle() const { return SliderHandleEndAngle; }
+
+	/** Get the AngularOffset attribute */
+	float GetAngularOffset() const { return AngularOffset; }
+
+	/** Get the ValueTags attribute */
+	TArray<float> GetValueTags() const { return ValueTags;	}
+
+	/** Get the Value attribute */
+	float GetValue() const;
+
+	/** Get the bUseCustomDefaultValue attribute */
+	bool GetUseCustomDefaultValue() const;
+
+	/** Get ths CustomDefaultValue attribute */
+	float GetCustomDefaultValue() const;
+
+	/** Get the Value attribute scaled from 0 to 1 */
+	float GetNormalizedValue(float RawValue) const;
+
+	/** Get the Slider's Handle position scaled from 0 to 1 */
+	float GetNormalizedSliderHandlePosition() const;
+
+	/** Set the Value attribute */
+	void SetValue(const TAttribute<float>& InValueAttribute);
+
+	/** Set the bUseCustomDefaultValue attribute */
+	void SetUseCustomDefaultValue(const TAttribute<bool>& InValueAttribute);
+
+	/** Set the CustomDefaultValue attribute */
+	void SetCustomDefaultValue(const TAttribute<float>& InValueAttribute);
+	
+	/** Set the SliderRange attribute */
+	void SetSliderRange(const FRuntimeFloatCurve& InSliderRange) { SliderRange = InSliderRange; }
+
+	/** Set the SliderHandleStartAngle and SliderHandleEndAngle attributes. If the new SliderHandleStartAngle is more than the new SliderHandleEndAngle, SliderHandleEndAngle will be changed to equal SliderHandleStartAngle. */
+	void SetSliderHandleStartAngleAndSliderHandleEndAngle(float InSliderHandleStartAngle, float InSliderHandleEndAngle);
+	
+	/** Set the AngularOffset attribute */
+	void SetAngularOffset(float InAngularOffset) { AngularOffset = InAngularOffset; }
+
+	/** Set the ValueTags attribute */
+	void SetValueTags(const TArray<float>& InValueTags) { ValueTags = InValueTags; }
+
+	/** Set the Locked attribute */
+	void SetLocked(const TAttribute<bool>& InLocked);
+	
+	/** Set the SliderBarColor attribute */
+	void SetSliderBarColor(FSlateColor InSliderBarColor);
+	
+	/** Set the SliderProgressColor attribute */
+	void SetSliderProgressColor(FSlateColor InSliderProgressColor);
+
+	/** Set the SliderHandleColor attribute */
+	void SetSliderHandleColor(FSlateColor InSliderHandleColor);
+
+	/** Get the StepSize attribute */
+	float GetStepSize() const;
+
+	/** Set the StepSize attribute */
+	void SetStepSize(const TAttribute<float>& InStepSize);
+
+	/** Set the MouseUsesStep attribute */
+	void SetMouseUsesStep(bool MouseUsesStep);
+
+	/** Set the RequiresControllerLock attribute */
+	void SetRequiresControllerLock(bool RequiresControllerLock);
+
+public:
+
+	// SWidget overrides
+
+	virtual int32 OnPaint( const FPaintArgs& Args, const FGeometry& AllottedGeometry, const FSlateRect& MyCullingRect, FSlateWindowElementList& OutDrawElements, int32 LayerId, const FWidgetStyle& InWidgetStyle, bool bParentEnabled ) const override;
+	virtual FVector2D ComputeDesiredSize(float) const override;
+	virtual FReply OnMouseButtonDown( const FGeometry& MyGeometry, const FPointerEvent& MouseEvent ) override;
+	virtual FReply OnMouseButtonUp( const FGeometry& MyGeometry, const FPointerEvent& MouseEvent ) override;
+	virtual FReply OnMouseMove( const FGeometry& MyGeometry, const FPointerEvent& MouseEvent ) override;
+	virtual FReply OnTouchStarted(const FGeometry& MyGeometry, const FPointerEvent& InTouchEvent) override;
+	virtual FReply OnTouchMoved(const FGeometry& MyGeometry, const FPointerEvent& InTouchEvent) override;
+	virtual FReply OnTouchEnded(const FGeometry& MyGeometry, const FPointerEvent& InTouchEvent) override;
+	virtual FReply OnKeyDown(const FGeometry& MyGeometry, const FKeyEvent& InKeyEvent) override;
+	virtual FReply OnKeyUp(const FGeometry& MyGeometry, const FKeyEvent& InKeyEvent) override;
+	virtual void OnFocusLost(const FFocusEvent& InFocusEvent) override;
+	virtual FNavigationReply OnNavigation(const FGeometry& MyGeometry, const FNavigationEvent& InNavigationEvent) override;
+
+	virtual bool SupportsKeyboardFocus() const override;
+	virtual bool IsInteractable() const override;
+#if WITH_ACCESSIBILITY
+	virtual TSharedRef<FSlateAccessibleWidget> CreateAccessibleWidget() override;
+#endif
+
+	/** @return Is the handle locked or not? Defaults to false */
+	bool IsLocked() const;
+
+protected:
+
+	/**
+	 * Commits the specified slider value.
+	 *
+	 * @param NewValue The value to commit.
+	 */
+	virtual void CommitValue(float NewValue);
+
+	/**
+	 * Calculates the new value based on the given absolute coordinates.
+	 *
+	 * @param MyGeometry The slider's geometry.
+	 * @param AbsolutePosition The absolute position of the slider.
+	 * @return The new value.
+	 */
+	float PositionToValue(const FGeometry& MyGeometry, const FVector2D& AbsolutePosition);
+
+	const FSlateBrush* GetBarImage() const;
+	const FSlateBrush* GetThumbImage() const;
+
+protected:
+
+	// Holds the style passed to the widget upon construction.
+	const FSliderStyle* Style;
+
+	// Holds a flag indicating whether the slider is locked.
+	TAttribute<bool> LockedAttribute;
+
+	// Holds the color of the slider bar.
+	TAttribute<FSlateColor> SliderBarColor;
+
+	// Holds the color for the completed progress of the slider bar.
+	TAttribute<FSlateColor> SliderProgressColor;
+
+	// Holds the color of the slider handle.
+	TAttribute<FSlateColor> SliderHandleColor;
+
+	// Holds the slider's current value.
+	TAttribute<float> ValueAttribute;
+
+	// Whether the slider should draw it's progress bar from a custom value on the slider
+	TAttribute<bool> bUseCustomDefaultValue;
+
+	// The value where the slider should draw it's progress bar from, independent of direction
+	TAttribute<float> CustomDefaultValue;
+
+	// Holds the initial cursor in case a custom cursor has been specified, so we can restore it after dragging the slider
+	EMouseCursor::Type CachedCursor;
+
+	/** The location in screenspace the slider was pressed by a touch */
+	FVector2D PressedScreenSpaceTouchDownPosition = FVector2D(0, 0);
+
+	/** Holds the amount to adjust the value by when using a controller or keyboard */
+	TAttribute<float> StepSize;
+
+	/**  The angle at which the radial slider should begin */	
+	float SliderHandleStartAngle;
+
+	/**  The angle at which the radial slider should end */	
+	float SliderHandleEndAngle;
+
+	/**  The angle at which the radial slider should be offset by */	
+	float AngularOffset;
+
+	/**  The values that should be drawn around the radial slider*/	
+	TArray<float> ValueTags;
+
+	/** A curve that defines how the slider should be sampled. Default is linear. */	
+	FRuntimeFloatCurve SliderRange;
+
+	// Holds a flag indicating whether a controller/keyboard is manipulating the slider's value. 
+	// When true, navigation away from the widget is prevented until a new value has been accepted or canceled. 
+	bool bControllerInputCaptured;
+
+	/** Sets new value if mouse position is greater/less than half the step size. */
+	bool bMouseUsesStep;
+
+	/** Sets whether we have to lock input to change the slider value. */
+	bool bRequiresControllerLock;
+
+	/** When true, this slider will be keyboard focusable. Defaults to false. */
+	bool bIsFocusable;
+
+private:
+
+	// Resets controller input state. Fires delegates.
+	void ResetControllerState();
+
+	// Helper function to convert Absolute Position to Angle
+	float GetAngleFromPosition(const FGeometry& MyGeometry, const FVector2D& AbsolutePosition);
+	
+	// Called on Mouse / Touch input to cache relevant properties
+	void OnInputStarted(const FGeometry& MyGeometry, const FVector2D& InputAbsolutePosition);
+
+	// Helper function for adding slider points to a slider points array
+	void AddSliderPointToArray(TArray<FVector2D>& SliderPoints, const bool bIsUnique, const FVector2D& SliderPoint) const;
+
+	// Holds a delegate that is executed when the mouse is pressed and a capture begins.
+	FSimpleDelegate OnMouseCaptureBegin;
+
+	// Holds a delegate that is executed when the mouse is let up and a capture ends.
+	FSimpleDelegate OnMouseCaptureEnd;
+
+	// Holds a delegate that is executed when capture begins for controller or keyboard.
+	FSimpleDelegate OnControllerCaptureBegin;
+
+	// Holds a delegate that is executed when capture ends for controller or keyboard.
+	FSimpleDelegate OnControllerCaptureEnd;
+
+	// Holds a delegate that is executed when the slider's value changed.
+	FOnFloatValueChanged OnValueChanged;
+
+	// Holds the current interaction's unclamped input angle.
+	float AbsoluteInputAngle;
+
+	// Stores the previous absolute position to support calculating rotational delta for relative input
+	FVector2D PreviousAbsolutePosition;
+};
