Fix for WinGDK missing includes

diff --git a/Engine/Source/Runtime/AdvancedWidgets/Private/Slate/SRadialSlider.cpp b/Engine/Source/Runtime/AdvancedWidgets/Private/Slate/SRadialSlider.cpp
--- a/Engine/Source/Runtime/AdvancedWidgets/Private/Slate/SRadialSlider.cpp
+++ b/Engine/Source/Runtime/AdvancedWidgets/Private/Slate/SRadialSlider.cpp
@@ -2,11 +2,12 @@
 
 #include "Slate/SRadialSlider.h"
 
+#include "Fonts/FontMeasure.h"
+#include "Fonts/SlateFontInfo.h"
 #include "Framework/Application/SlateApplication.h"
 #include "Rendering/DrawElements.h"
+
 #if WITH_ACCESSIBILITY
-#include "Fonts/SlateFontInfo.h"
-#include "Fonts/FontMeasure.h"
 #include "Widgets/Accessibility/SlateAccessibleWidgets.h"
 #endif
 
