AdvancedWidgets: Runtime widget to display the property value of a UObject. Features include:
  - Can show the value of the properties (false by default).
  - Can show functions and/or properties. By default, it shows everything that is BP visible but it can be overridden with a callback.
   - It uses the BP VM to set the value (opposite to the DetailView that uses the editor pipeline to re-instantiate the child of the object).
   - User can define an "customizer"  for your type. If not customizer is found, it will export to text the property value. The "customizer" is a simple SWidget. There is no super custom functionality like the DetailView has.
  - Search bar (can be customized with additional widgets).

diff --git a/Engine/Source/Runtime/AdvancedWidgets/AdvancedWidgets.Build.cs b/Engine/Source/Runtime/AdvancedWidgets/AdvancedWidgets.Build.cs
--- a/Engine/Source/Runtime/AdvancedWidgets/AdvancedWidgets.Build.cs
+++ b/Engine/Source/Runtime/AdvancedWidgets/AdvancedWidgets.Build.cs
@@ -6,17 +6,27 @@ public class AdvancedWidgets : ModuleRules
 {
 	public AdvancedWidgets(ReadOnlyTargetRules Target) : base(Target)
 	{
-		PrivateDependencyModuleNames.AddRange(
+		PublicDependencyModuleNames.AddRange(
 			new string[]
 			{
 				"Core",
 				"CoreUObject",
+				"SlateCore",
+				"UMG"
+			}
+		);
+		PrivateDependencyModuleNames.AddRange(
+			new string[]
+			{
 				"Engine",
 				"InputCore",
 				"Slate",
-				"SlateCore",
-				"UMG"
 			}
 		);
+
+		if (Target.bBuildEditor == true)
+		{
+			PrivateDependencyModuleNames.Add("UnrealEd");
+		}
 	}
 }
diff --git a/Engine/Source/Runtime/AdvancedWidgets/Private/AdvancedWidgetsModule.cpp b/Engine/Source/Runtime/AdvancedWidgets/Private/AdvancedWidgetsModule.cpp
--- a/Engine/Source/Runtime/AdvancedWidgets/Private/AdvancedWidgetsModule.cpp
+++ b/Engine/Source/Runtime/AdvancedWidgets/Private/AdvancedWidgetsModule.cpp
@@ -2,16 +2,40 @@
 
 #include "AdvancedWidgetsModule.h"
 
+#include "Framework/PropertyViewer/PropertyValueFactory.h"
 #include "Modules/ModuleManager.h"
+#include "Styling/AdvancedWidgetsStyle.h"
+
 
 IMPLEMENT_MODULE(FAdvancedWidgetsModule, AdvancedWidgets);
 
-const FName FAdvancedWidgetsModule::AdvancedWidgetsAppIdentifier = "AdvancedWidgetsApp";
+namespace UE::AdvancedWidgets::Private
+{
+	static const FName NAME_AdvancedWidgets = "AdvancedWidgets";
+}
+
+FAdvancedWidgetsModule& FAdvancedWidgetsModule::GetModule()
+{
+	return FModuleManager::LoadModuleChecked<FAdvancedWidgetsModule>(UE::AdvancedWidgets::Private::NAME_AdvancedWidgets);
+}
+
+FAdvancedWidgetsModule* FAdvancedWidgetsModule::GetModulePtr()
+{
+	return FModuleManager::GetModulePtr<FAdvancedWidgetsModule>(UE::AdvancedWidgets::Private::NAME_AdvancedWidgets);
+}
 
 void FAdvancedWidgetsModule::StartupModule()
 {
+	UE::AdvancedWidgets::FAdvancedWidgetsStyle::Create();
+	PropertyValueFactory = MakeUnique<UE::PropertyViewer::FPropertyValueFactory>();
 }
 
 void FAdvancedWidgetsModule::ShutdownModule()
 {
+	UE::AdvancedWidgets::FAdvancedWidgetsStyle::Destroy();
+}
+
+UE::PropertyViewer::FPropertyValueFactory& FAdvancedWidgetsModule::GetPropertyValueFactory() const
+{
+	return *(PropertyValueFactory.Get());
 }
diff --git a/Engine/Source/Runtime/AdvancedWidgets/Private/Framework/PropertyViewer/FieldIconFinder.cpp b/Engine/Source/Runtime/AdvancedWidgets/Private/Framework/PropertyViewer/FieldIconFinder.cpp
new file mode 100644
--- /dev/null
+++ b/Engine/Source/Runtime/AdvancedWidgets/Private/Framework/PropertyViewer/FieldIconFinder.cpp
@@ -0,0 +1,195 @@
+// Copyright Epic Games, Inc. All Rights Reserved.
+
+#include "Framework/PropertyViewer/FieldIconFinder.h"
+
+#include "Styling/AdvancedWidgetsStyle.h"
+#include "Styling/AppStyle.h"
+#include "Styling/SlateBrush.h"
+#include "Styling/SlateIconFinder.h"
+#include "UObject/Class.h"
+#include "UObject/EnumProperty.h"
+#include "UObject/TextProperty.h"
+#include "UObject/UnrealType.h"
+
+
+namespace UE::PropertyViewer
+{
+
+namespace Private
+{
+
+FSlateColor GetColor(const FFieldColorSettings& Settings, const FProperty* Property)
+{
+	if (CastField<const FClassProperty>(Property))
+	{
+		return Settings.ClassTypeColor;
+	}
+	if (CastField<const FObjectPropertyBase>(Property))
+	{
+		return Settings.ObjectTypeColor;
+	}
+	if (CastField<const FInterfaceProperty>(Property))
+	{
+		return Settings.InterfaceTypeColor;
+	}
+	if (const FNumericProperty* NumericProperty = CastField<const FNumericProperty>(Property))
+	{
+		if (NumericProperty->IsEnum())
+		{
+			return Settings.EnumTypeColor;
+		}
+		if (NumericProperty->IsFloatingPoint())
+		{
+			return Settings.FloatTypeColor;
+		}
+		if (NumericProperty->IsInteger())
+		{
+			return Settings.IntTypeColor;
+		}
+	}
+	if (CastField<const FBoolProperty>(Property))
+	{
+		return Settings.BooleanTypeColor;
+	}
+	if (CastField<const FEnumProperty>(Property))
+	{
+		return Settings.EnumTypeColor;
+	}
+	if (CastField<const FStrProperty>(Property))
+	{
+		return Settings.StringTypeColor;
+	}
+	if (CastField<const FTextProperty>(Property))
+	{
+		return Settings.StringTypeColor;
+	}
+	if (CastField<const FNameProperty>(Property))
+	{
+		return Settings.TextTypeColor;
+	}
+	if (CastField<const FMulticastDelegateProperty>(Property) || CastField<const FDelegateProperty>(Property))
+	{
+		return Settings.DelegateTypeColor;
+	}
+	if (const FStructProperty* StructProperty = CastField<const FStructProperty>(Property))
+	{
+		if (const FLinearColor* FoundColor = Settings.StructColors.Find(StructProperty->Struct->GetStructPathName().ToString()))
+		{
+			return *FoundColor;
+		}
+		return Settings.DefaultStructTypeColor;
+	}
+	return Settings.DefaultTypeColor;
+}
+
+} //namespace Private
+
+
+const FSlateBrush* FFieldIconFinder::GetIcon(const UObject* Object)
+{
+	if (const UFunction* Function = Cast<const UFunction>(Object))
+	{
+		FFieldIconArray FunctionResult = GetFunctionIcon(Function);
+		if (FunctionResult.Num() > 0)
+		{
+			return FunctionResult[0].Icon;
+		}
+	}
+	if (const UClass* Class = Cast<const UClass>(Object))
+	{
+		return FSlateIconFinder::FindIconBrushForClass(Class);
+	}
+	return FSlateIconFinder::FindIconBrushForClass(Object->GetClass());
+}
+
+
+FFieldIconFinder::FFieldIconArray FFieldIconFinder::GetFunctionIcon(const UFunction* Function)
+{
+	return GetFunctionIcon(Function, ::UE::AdvancedWidgets::FAdvancedWidgetsStyle::GetColorSettings());
+}
+
+
+FFieldIconFinder::FFieldIconArray FFieldIconFinder::GetFunctionIcon(const UFunction* Function, const FFieldColorSettings& Settings)
+{
+	FFieldIcon Icon;
+	if (UFunction* OverrideFunc = Function->GetSuperFunction())
+	{
+		const bool bIsPureFunction = OverrideFunc && OverrideFunc->HasAnyFunctionFlags(FUNC_BlueprintPure);
+		Icon.Icon = FAppStyle::GetBrush(bIsPureFunction ? TEXT("GraphEditor.OverridePureFunction_16x") : TEXT("GraphEditor.OverrideFunction_16x"));
+	}
+	else
+	{
+		const bool bIsPureFunction = Function && Function->HasAnyFunctionFlags(FUNC_BlueprintPure);
+		Icon.Icon = FAppStyle::GetBrush(bIsPureFunction ? TEXT("GraphEditor.PureFunction_16x") : TEXT("GraphEditor.Function_16x"));
+	}
+
+	Icon.Color = FLinearColor::White;
+	if (FProperty* ReturnProperty = Function->GetReturnProperty())
+	{
+		FFieldIconArray ReturnValueResult = GetPropertyIcon(ReturnProperty, Settings);
+		if (ReturnValueResult.Num() > 0)
+		{
+			Icon.Color = ReturnValueResult[0].Color;
+		}
+	}
+
+	FFieldIconArray Result;
+	Result.Add(Icon);
+	return Result;
+}
+
+
+FFieldIconFinder::FFieldIconArray FFieldIconFinder::GetPropertyIcon(const FProperty* Property)
+{
+	return GetPropertyIcon(Property, ::UE::AdvancedWidgets::FAdvancedWidgetsStyle::GetColorSettings());
+}
+
+
+FFieldIconFinder::FFieldIconArray FFieldIconFinder::GetPropertyIcon(const FProperty* Property, const FFieldColorSettings& Settings)
+{
+	FFieldIconArray Result;
+
+	if (const FArrayProperty* ArrayProperty = CastField<FArrayProperty>(Property))
+	{
+		FFieldIcon Icon;
+		Icon.Icon = FAppStyle::GetBrush(TEXT("Kismet.VariableList.ArrayTypeIcon"));
+		Icon.Color = Private::GetColor(Settings, ArrayProperty->Inner);
+		Result.Add(Icon);
+	}
+	else if (const FMapProperty* MapProperty = CastField<FMapProperty>(Property))
+	{
+		FFieldIcon Icon;
+		Icon.Icon = FAppStyle::GetBrush(TEXT("Kismet.VariableList.MapKeyTypeIcon"));
+		Icon.Color = Private::GetColor(Settings, MapProperty->GetKeyProperty());
+		Result.Add(Icon);
+		Icon.Icon = FAppStyle::GetBrush(TEXT("Kismet.VariableList.MapValueTypeIcon"));
+		Icon.Color = Private::GetColor(Settings, MapProperty->GetValueProperty());
+		Result.Add(Icon);
+	}
+	else if (const FSetProperty* SetProperty = CastField<FSetProperty>(Property))
+	{
+		FFieldIcon Icon;
+		Icon.Icon = FAppStyle::GetBrush(TEXT("Kismet.VariableList.SetTypeIcon"));
+		Icon.Color = Private::GetColor(Settings, SetProperty->ElementProp);
+		Result.Add(Icon);
+	}
+	else if (const FObjectPropertyBase* ObjectProperty = CastField<FObjectPropertyBase>(Property))
+	{
+		FFieldIcon Icon;
+		Icon.Icon = FSlateIconFinder::FindIconBrushForClass(ObjectProperty->PropertyClass);
+		Icon.Color = Private::GetColor(Settings, Property);
+		Result.Add(Icon);
+	}
+	else
+	{
+		FFieldIcon Icon;
+		Icon.Icon = FAppStyle::GetBrush(TEXT("Kismet.VariableList.TypeIcon"));
+		Icon.Color = Private::GetColor(Settings, Property);
+		Result.Add(Icon);
+	}
+
+	//return FAppStyle::GetBrush(TEXT("Kismet.AllClasses.VariableIcon");
+	return Result;
+}
+
+} //namespace
diff --git a/Engine/Source/Runtime/AdvancedWidgets/Private/Framework/PropertyViewer/IFieldIterator.cpp b/Engine/Source/Runtime/AdvancedWidgets/Private/Framework/PropertyViewer/IFieldIterator.cpp
new file mode 100644
--- /dev/null
+++ b/Engine/Source/Runtime/AdvancedWidgets/Private/Framework/PropertyViewer/IFieldIterator.cpp
@@ -0,0 +1,37 @@
+// Copyright Epic Games, Inc. All Rights Reserved.
+
+#include "Framework/PropertyViewer/IFieldIterator.h"
+
+#include "UObject/Class.h"
+#include "UObject/UnrealType.h"
+
+namespace UE::PropertyViewer
+{
+
+/**
+ * FFieldIterator_BlueprintVisible
+ */
+TArray<FFieldVariant> FFieldIterator_BlueprintVisible::GetFields(const UStruct* Struct) const
+{
+	TArray<FFieldVariant> Result;
+	for (TFieldIterator<FProperty> PropertyIt(Struct, EFieldIteratorFlags::IncludeSuper); PropertyIt; ++PropertyIt)
+	{
+		FProperty* Property = *PropertyIt;
+		if (Property->HasAnyPropertyFlags(CPF_BlueprintVisible | CPF_BlueprintAssignable | CPF_Parm)
+			&& !Property->HasAnyPropertyFlags(CPF_Deprecated | CPF_EditorOnly))
+		{
+			Result.Add(FFieldVariant(Property));
+		}
+	}
+	for (TFieldIterator<UFunction> FunctionItt(Struct, EFieldIteratorFlags::IncludeSuper); FunctionItt; ++FunctionItt)
+	{
+		UFunction* Function = *FunctionItt;
+		if (Function->HasAnyFunctionFlags(FUNC_BlueprintCallable))
+		{
+			Result.Add(FFieldVariant(Function));
+		}
+	}
+	return Result;
+}
+
+} //namespace
diff --git a/Engine/Source/Runtime/AdvancedWidgets/Private/Framework/PropertyViewer/PropertyPath.cpp b/Engine/Source/Runtime/AdvancedWidgets/Private/Framework/PropertyViewer/PropertyPath.cpp
new file mode 100644
--- /dev/null
+++ b/Engine/Source/Runtime/AdvancedWidgets/Private/Framework/PropertyViewer/PropertyPath.cpp
@@ -0,0 +1,90 @@
+// Copyright Epic Games, Inc. All Rights Reserved.
+
+#include "Framework/PropertyViewer/PropertyPath.h"
+#include "UObject/UnrealType.h"
+
+namespace UE::PropertyViewer
+{
+
+FPropertyPath::FPropertyPath(UObject* Object, const FProperty* Property)
+	: TopLevelContainer_Object(Object)
+{
+	Properties.Add(Property);
+}
+
+
+FPropertyPath::FPropertyPath(UObject* Object, FPropertyArray InProperties)
+	: TopLevelContainer_Object(Object)
+	, Properties(MoveTemp(InProperties))
+{}
+
+
+FPropertyPath::FPropertyPath(const UScriptStruct* ScriptStruct, void* Data, const FProperty* Property)
+	: TopLevelContainer_ScriptStruct(ScriptStruct)
+	, TopLevelContainer_ScriptStructData(Data)
+{
+	Properties.Add(Property);
+}
+
+
+FPropertyPath::FPropertyPath(const UScriptStruct* ScriptStruct, void* Data, FPropertyArray InProperties)
+	: TopLevelContainer_ScriptStruct(ScriptStruct)
+	, TopLevelContainer_ScriptStructData(Data)
+	, Properties(MoveTemp(InProperties))
+{}
+
+
+void* FPropertyPath::GetContainerPtr()
+{
+	void* Current = nullptr;
+	if (UObject* ObjectInstance = TopLevelContainer_Object.Get())
+	{
+		bool bInvalidClass = ObjectInstance->GetClass()->HasAnyClassFlags(EClassFlags::CLASS_NewerVersionExists);
+		if (!bInvalidClass)
+		{
+			Current = reinterpret_cast<void*>(ObjectInstance);
+		}
+	}
+	else if (const UScriptStruct* StructInstance = TopLevelContainer_ScriptStruct.Get())
+	{
+		bool bInvalidClass = (StructInstance->StructFlags & (EStructFlags::STRUCT_Trashed)) != 0;
+		if (!bInvalidClass)
+		{
+			Current = TopLevelContainer_ScriptStructData;
+		}
+	}
+
+	for (int32 Index = 0; Index < Properties.Num() - 1; ++Index)
+	{
+		if (Current == nullptr)
+		{
+			return nullptr;
+		}
+
+		const FProperty* CurrentProperty = Properties[Index];
+		if (CurrentProperty == nullptr)
+		{
+			return nullptr;
+		}
+
+		if (const FObjectPropertyBase* ObjectProperty = CastField<const FObjectPropertyBase>(CurrentProperty))
+		{
+			Current = reinterpret_cast<void*>(ObjectProperty->GetObjectPropertyValue_InContainer(Current));
+		}
+		else
+		{
+			// We do not use the Getter on purpose.
+			Current = CurrentProperty->ContainerPtrToValuePtr<void>(Current);
+		}
+	}
+
+	return Current;
+}
+
+
+const void* FPropertyPath::GetContainerPtr() const
+{
+	return const_cast<FPropertyPath*>(this)->GetContainerPtr();
+}
+
+} //namespace
diff --git a/Engine/Source/Runtime/AdvancedWidgets/Private/Framework/PropertyViewer/PropertyValueFactory.cpp b/Engine/Source/Runtime/AdvancedWidgets/Private/Framework/PropertyViewer/PropertyValueFactory.cpp
new file mode 100644
--- /dev/null
+++ b/Engine/Source/Runtime/AdvancedWidgets/Private/Framework/PropertyViewer/PropertyValueFactory.cpp
@@ -0,0 +1,212 @@
+// Copyright Epic Games, Inc. All Rights Reserved.
+
+#include "Framework/PropertyViewer/PropertyValueFactory.h"
+#include "UObject/EnumProperty.h"
+#include "UObject/TextProperty.h"
+#include "UObject/TopLevelAssetPath.h"
+#include "UObject/UnrealType.h"
+
+#include "Widgets/PropertyViewer/SBoolPropertyValue.h"
+#include "Widgets/PropertyViewer/SEnumPropertyValue.h"
+#include "Widgets/PropertyViewer/SDefaultPropertyValue.h"
+#include "Widgets/PropertyViewer/SNumericPropertyValue.h"
+#include "Widgets/PropertyViewer/SStringPropertyValue.h"
+
+
+namespace UE::PropertyViewer
+{
+namespace Private
+{
+
+TSharedPtr<SWidget> GenerateNumericEditor(FPropertyValueFactory::FGenerateArgs Args)
+{
+	if (const FProperty* Property = Args.Path.GetLastProperty())
+	{
+		if (CastFieldChecked<const FNumericProperty>(Property)->IsEnum())
+		{
+			return SEnumPropertyValue::CreateInstance(Args);
+		}
+	}
+	return SNumericPropertyValue::CreateInstance(Args);
+}
+
+
+struct FPropertyViewerMapItem
+{
+	FPropertyValueFactory::FOnGenerate OnGenerate;
+	FPropertyValueFactory::FHandle Handle;
+};
+
+
+class FPropertyValueFactoryImpl
+{
+public:
+	/** For FProperty. Use FFieldClass::Name */
+	TMap<const FFieldClass*, FPropertyViewerMapItem> RegisteredFieldClassEditor;
+	/** For UScriptStruct or UClass */
+	TMap<FTopLevelAssetPath, FPropertyViewerMapItem> RegisteredFieldEditor;
+};
+
+} //namespace
+
+
+FPropertyValueFactory::FPropertyValueFactory()
+{
+	Impl = MakePimpl<Private::FPropertyValueFactoryImpl>();
+	Register(FNumericProperty::StaticClass(), FPropertyValueFactory::FOnGenerate::CreateStatic(Private::GenerateNumericEditor));
+	Register(FBoolProperty::StaticClass(), FPropertyValueFactory::FOnGenerate::CreateStatic(SBoolPropertyValue::CreateInstance));
+	Register(FEnumProperty::StaticClass(), FPropertyValueFactory::FOnGenerate::CreateStatic(SEnumPropertyValue::CreateInstance));
+	Register(FStrProperty::StaticClass(), FPropertyValueFactory::FOnGenerate::CreateStatic(SStringPropertyValue::CreateInstance));
+	Register(FTextProperty::StaticClass(), FPropertyValueFactory::FOnGenerate::CreateStatic(SStringPropertyValue::CreateInstance));
+}
+
+
+TSharedPtr<SWidget> FPropertyValueFactory::Generate(FGenerateArgs Args) const
+{
+	if (const FProperty* LastProperty = Args.Path.GetLastProperty())
+	{
+		const bool bIsVisible = LastProperty->HasAllPropertyFlags(CPF_BlueprintVisible);
+		if (bIsVisible && LastProperty->ArrayDim == 1)
+		{
+			if (LastProperty->HasAnyPropertyFlags(CPF_BlueprintReadOnly))
+			{
+				Args.bCanEditValue = false;
+			}
+
+			if (const FObjectPropertyBase* ObjectProperty = CastField<const FObjectPropertyBase>(LastProperty))
+			{
+				if (Private::FPropertyViewerMapItem* FoundOnGenerate = Impl->RegisteredFieldEditor.Find(ObjectProperty->PropertyClass->GetStructPathName()))
+				{
+					check(FoundOnGenerate->OnGenerate.IsBound());
+					return FoundOnGenerate->OnGenerate.Execute(Args);
+				}
+			}
+			else if (const FStructProperty* StructProperty = CastField<const FStructProperty>(LastProperty))
+			{
+				if (Private::FPropertyViewerMapItem* FoundOnGenerate = Impl->RegisteredFieldEditor.Find(StructProperty->Struct->GetStructPathName()))
+				{
+					check(FoundOnGenerate->OnGenerate.IsBound());
+					return FoundOnGenerate->OnGenerate.Execute(Args);
+				}
+			}
+			else
+			{
+				if (Private::FPropertyViewerMapItem* FoundOnGenerate = Impl->RegisteredFieldClassEditor.Find(LastProperty->GetClass()))
+				{
+					check(FoundOnGenerate->OnGenerate.IsBound());
+					return FoundOnGenerate->OnGenerate.Execute(Args);
+				}
+				for (auto& ClassEditor : Impl->RegisteredFieldClassEditor)
+				{
+					if (LastProperty->GetClass()->IsChildOf(ClassEditor.Key))
+					{
+						check(ClassEditor.Value.OnGenerate.IsBound());
+						return ClassEditor.Value.OnGenerate.Execute(Args);
+					}
+				}
+			}
+		}
+	}
+	return TSharedPtr<SWidget>();
+}
+
+
+TSharedPtr<SWidget> FPropertyValueFactory::GenerateDefault(FGenerateArgs Args) const
+{
+	if (const FProperty* LastProperty = Args.Path.GetLastProperty())
+	{
+		const bool bIsVisible = LastProperty->HasAllPropertyFlags(CPF_BlueprintVisible);
+		if (bIsVisible)
+		{
+			if (LastProperty->HasAnyPropertyFlags(CPF_BlueprintReadOnly))
+			{
+				Args.bCanEditValue = false;
+			}
+
+			return SDefaultPropertyValue::CreateInstance(Args);
+		}
+	}
+	return TSharedPtr<SWidget>();
+}
+
+
+bool FPropertyValueFactory::HasCustomPropertyValue(const FFieldClass* FieldClass) const
+{
+	check(FieldClass);
+	return Impl->RegisteredFieldClassEditor.Find(FieldClass) != nullptr;
+}
+
+
+bool FPropertyValueFactory::HasCustomPropertyValue(const UStruct* Struct) const
+{
+	check(Struct);
+	return Impl->RegisteredFieldEditor.Find(Struct->GetStructPathName()) != nullptr;
+}
+
+
+FPropertyValueFactory::FHandle FPropertyValueFactory::Register(const FFieldClass* FieldClass, FOnGenerate OnGenerateFieldEditor)
+{
+	if (!FieldClass || !OnGenerateFieldEditor.IsBound())
+	{
+		return FHandle();
+	}
+
+	FPropertyValueFactory::FHandle Result = MakeHandle();
+	Private::FPropertyViewerMapItem Item;
+	Item.OnGenerate = MoveTemp(OnGenerateFieldEditor);
+	Item.Handle = Result;
+
+	Impl->RegisteredFieldClassEditor.FindOrAdd(FieldClass) = MoveTemp(Item);
+
+	return Result;
+}
+
+
+FPropertyValueFactory::FHandle FPropertyValueFactory::Register(const UStruct* Struct, FOnGenerate OnGenerateFieldEditor)
+{
+	if (!Struct || !OnGenerateFieldEditor.IsBound())
+	{
+		return FHandle();
+	}
+
+	FPropertyValueFactory::FHandle Result = MakeHandle();
+	Private::FPropertyViewerMapItem Item;
+	Item.OnGenerate = MoveTemp(OnGenerateFieldEditor);
+	Item.Handle = Result;
+
+	Impl->RegisteredFieldEditor.FindOrAdd(Struct->GetStructPathName()) = MoveTemp(Item);
+	return Result;
+}
+
+
+void FPropertyValueFactory::Unregister(FHandle Handle)
+{
+	for (auto It = Impl->RegisteredFieldEditor.CreateIterator(); It; ++It)
+	{
+		if (It->Value.Handle == Handle)
+		{
+			It.RemoveCurrent();
+			return;
+		}
+	}
+	for (auto It = Impl->RegisteredFieldClassEditor.CreateIterator(); It; ++It)
+	{
+		if (It->Value.Handle == Handle)
+		{
+			It.RemoveCurrent();
+			return;
+		}
+	}
+}
+
+
+FPropertyValueFactory::FHandle  FPropertyValueFactory::MakeHandle()
+{
+	static int32 Generator = 0;
+	++Generator;
+	FHandle Result;
+	Result.Id = Generator;
+	return Result;
+}
+
+} //namespace
diff --git a/Engine/Source/Runtime/AdvancedWidgets/Private/Styling/AdvancedWidgetsStyle.cpp b/Engine/Source/Runtime/AdvancedWidgets/Private/Styling/AdvancedWidgetsStyle.cpp
new file mode 100644
--- /dev/null
+++ b/Engine/Source/Runtime/AdvancedWidgets/Private/Styling/AdvancedWidgetsStyle.cpp
@@ -0,0 +1,77 @@
+// Copyright Epic Games, Inc. All Rights Reserved.
+
+#include "Styling/AdvancedWidgetsStyle.h"
+
+#include "Framework/PropertyViewer/FieldIconFinder.h"
+#include "Math/ColorList.h"
+#include "Styling/AppStyle.h"
+#include "Styling/SlateColor.h"
+#include "Styling/SlateStyle.h"
+#include "Styling/SlateStyleRegistry.h"
+#include "Styling/SlateTypes.h"
+#include "UObject/Class.h"
+
+
+#define BORDER_BRUSH( RelativePath, ... ) FSlateBorderBrush( Instance->RootToContentDir( RelativePath, TEXT(".png") ), __VA_ARGS__ )
+
+namespace UE::AdvancedWidgets
+{
+
+TUniquePtr<FSlateStyleSet> FAdvancedWidgetsStyle::Instance;
+::UE::PropertyViewer::FFieldColorSettings FAdvancedWidgetsStyle::ColorSettings;
+
+
+const ISlateStyle& FAdvancedWidgetsStyle::Get()
+{
+	return *(Instance.Get());
+}
+
+void FAdvancedWidgetsStyle::Create()
+{
+	Instance = MakeUnique<FSlateStyleSet>("AdvancedWidgets");
+	Instance->SetContentRoot(FPaths::EngineContentDir() / TEXT("Slate"));
+	Instance->SetCoreContentRoot(FPaths::EngineContentDir() / TEXT("Slate"));
+
+	Instance->Set("PropertyValue.SpinBox", FSpinBoxStyle(FAppStyle::GetWidgetStyle<FSpinBoxStyle>("SpinBox"))
+		.SetTextPadding(FMargin(0.f))
+		.SetBackgroundBrush(BORDER_BRUSH("Old/HyperlinkDotted", FMargin(0.f, 0.f, 0.f, 3.f / 16.f), FSlateColor::UseSubduedForeground()))
+		.SetHoveredBackgroundBrush(FSlateNoResource())
+		.SetInactiveFillBrush(FSlateNoResource())
+		.SetActiveFillBrush(FSlateNoResource())
+		.SetForegroundColor(FSlateColor::UseSubduedForeground())
+		.SetArrowsImage(FSlateNoResource())
+	);
+
+	FButtonStyle Button = FAppStyle::GetWidgetStyle<FButtonStyle>("Button");
+	Instance->Set("PropertyValue.ComboButton", FComboButtonStyle(FAppStyle::GetWidgetStyle<FComboButtonStyle>("ComboButton"))
+		.SetButtonStyle(
+			Button.SetNormal(FSlateNoResource())
+			.SetNormalForeground(FSlateColor::UseSubduedForeground())
+			.SetDisabledForeground(FSlateColor::UseSubduedForeground())
+			.SetNormalPadding(FMargin(0.f))
+			.SetPressedPadding(FMargin(0.f))
+		)
+		.SetMenuBorderPadding(FMargin(0.0f))
+		.SetDownArrowPadding(FMargin(0.0f))
+		);
+
+	FLinearColor& VectorColor = ColorSettings.StructColors.FindOrAdd(TBaseStructure<FVector>::Get()->GetStructPathName().ToString());
+	VectorColor = FColorList::Yellow;
+	FLinearColor& RotatorColor = ColorSettings.StructColors.FindOrAdd(TBaseStructure<FRotator>::Get()->GetStructPathName().ToString());
+	RotatorColor = FColorList::DarkTurquoise;
+
+
+	FSlateStyleRegistry::RegisterSlateStyle(*Instance.Get());
+}
+
+void FAdvancedWidgetsStyle::Destroy()
+{
+	if (Instance)
+	{
+		FSlateStyleRegistry::UnRegisterSlateStyle(*Instance.Get());
+	}
+}
+
+} //namespace
+
+#undef BORDER_BRUSH
diff --git a/Engine/Source/Runtime/AdvancedWidgets/Private/Widgets/PropertyViewer/SBoolPropertyValue.cpp b/Engine/Source/Runtime/AdvancedWidgets/Private/Widgets/PropertyViewer/SBoolPropertyValue.cpp
new file mode 100644
--- /dev/null
+++ b/Engine/Source/Runtime/AdvancedWidgets/Private/Widgets/PropertyViewer/SBoolPropertyValue.cpp
@@ -0,0 +1,77 @@
+// Copyright Epic Games, Inc. All Rights Reserved.
+
+#include "Widgets/PropertyViewer/SBoolPropertyValue.h"
+
+#include "Framework/PropertyViewer/INotifyHook.h"
+#include "Styling/SlateTypes.h"
+#include "Styling/AppStyle.h"
+#include "UObject/TextProperty.h"
+#include "UObject/UnrealType.h"
+
+#include "Widgets/Input/SCheckBox.h"
+
+namespace UE::PropertyViewer
+{
+
+TSharedPtr<SWidget> SBoolPropertyValue::CreateInstance(const FPropertyValueFactory::FGenerateArgs Args)
+{
+	return SNew(SBoolPropertyValue)
+		.Path(Args.Path)
+		.NotifyHook(Args.NotifyHook)
+		.IsEnabled(Args.bCanEditValue);
+}
+
+
+void SBoolPropertyValue::Construct(const FArguments& InArgs)
+{
+	Path = InArgs._Path;
+	NotifyHook = InArgs._NotifyHook;
+
+	const FProperty* Property = Path.GetLastProperty();
+	if (CastField<const FBoolProperty>(Property) && Property->ArrayDim == 1)
+	{
+		ChildSlot
+		[
+			SNew(SCheckBox)
+			.IsChecked(this, &SBoolPropertyValue::HandleIsChecked)
+			.OnCheckStateChanged(this, &SBoolPropertyValue::HandleCheckStateChanged)
+ 		];
+	}
+}
+
+
+ECheckBoxState SBoolPropertyValue::HandleIsChecked() const
+{
+	if (const FBoolProperty* Property = CastField<const FBoolProperty>(Path.GetLastProperty()))
+	{
+		if (const void* Container = Path.GetContainerPtr())
+		{
+			const bool bValue = Property->GetPropertyValue_InContainer(Container);
+			return bValue ? ECheckBoxState::Checked : ECheckBoxState::Unchecked;
+		}
+	}
+	return ECheckBoxState::Undetermined;
+}
+
+
+void SBoolPropertyValue::HandleCheckStateChanged(ECheckBoxState NewState)
+{
+	if (const FBoolProperty* Property = CastField<const FBoolProperty>(Path.GetLastProperty()))
+	{
+		if (void* Container = Path.GetContainerPtr())
+		{
+			if (NotifyHook)
+			{
+				NotifyHook->OnPreValueChange(Path);
+			}
+			bool bValue = NewState == ECheckBoxState::Checked;
+			Property->SetPropertyValue_InContainer(Container, bValue);
+			if (NotifyHook)
+			{
+				NotifyHook->OnPostValueChange(Path);
+			}
+		}
+	}
+}
+
+} //namespace
diff --git a/Engine/Source/Runtime/AdvancedWidgets/Private/Widgets/PropertyViewer/SDefaultPropertyValue.cpp b/Engine/Source/Runtime/AdvancedWidgets/Private/Widgets/PropertyViewer/SDefaultPropertyValue.cpp
new file mode 100644
--- /dev/null
+++ b/Engine/Source/Runtime/AdvancedWidgets/Private/Widgets/PropertyViewer/SDefaultPropertyValue.cpp
@@ -0,0 +1,50 @@
+// Copyright Epic Games, Inc. All Rights Reserved.
+
+#include "Widgets/PropertyViewer/SDefaultPropertyValue.h"
+
+#include "Styling/SlateTypes.h"
+#include "Styling/AppStyle.h"
+#include "UObject/TextProperty.h"
+#include "UObject/UnrealType.h"
+
+#include "Widgets/Text/STextBlock.h"
+
+namespace UE::PropertyViewer
+{
+
+TSharedPtr<SWidget> SDefaultPropertyValue::CreateInstance(const FPropertyValueFactory::FGenerateArgs Args)
+{
+	return SNew(SDefaultPropertyValue)
+		.Path(Args.Path)
+		.IsEnabled(Args.bCanEditValue);
+}
+
+
+void SDefaultPropertyValue::Construct(const FArguments& InArgs)
+{
+	Path = InArgs._Path;
+
+	const FProperty* Property = Path.GetLastProperty();
+	ChildSlot
+	[
+		SNew(STextBlock)
+		.Text(this, &SDefaultPropertyValue::GetText)
+ 	];
+}
+
+
+FText SDefaultPropertyValue::GetText() const
+{
+	if (const FProperty* LastProperty = Path.GetLastProperty())
+	{
+		if (const void* ContainerPtr = Path.GetContainerPtr())
+		{
+			FString TextVersion;
+			LastProperty->ExportTextItem_InContainer(TextVersion, ContainerPtr, nullptr, nullptr, PPF_SimpleObjectText);
+			return FText::FromString(TextVersion);
+		}
+	}
+	return FText::GetEmpty();
+}
+
+} //namespace
diff --git a/Engine/Source/Runtime/AdvancedWidgets/Private/Widgets/PropertyViewer/SEnumPropertyValue.cpp b/Engine/Source/Runtime/AdvancedWidgets/Private/Widgets/PropertyViewer/SEnumPropertyValue.cpp
new file mode 100644
--- /dev/null
+++ b/Engine/Source/Runtime/AdvancedWidgets/Private/Widgets/PropertyViewer/SEnumPropertyValue.cpp
@@ -0,0 +1,182 @@
+// Copyright Epic Games, Inc. All Rights Reserved.
+
+#include "Widgets/PropertyViewer/SEnumPropertyValue.h"
+
+#include "Framework/MultiBox/MultiBoxBuilder.h"
+#include "Framework/PropertyViewer/INotifyHook.h"
+#include "Styling/SlateTypes.h"
+#include "Styling/AdvancedWidgetsStyle.h"
+#include "Styling/AppStyle.h"
+#include "UObject/TextProperty.h"
+#include "UObject/UnrealType.h"
+
+#include "Widgets/Input/SComboButton.h"
+#include "Widgets/Text/STextBlock.h"
+
+namespace UE::PropertyViewer
+{
+
+TSharedPtr<SWidget> SEnumPropertyValue::CreateInstance(const FPropertyValueFactory::FGenerateArgs Args)
+{
+	return SNew(SEnumPropertyValue)
+		.Path(Args.Path)
+		.NotifyHook(Args.NotifyHook)
+		.IsEnabled(Args.bCanEditValue);
+}
+
+
+void SEnumPropertyValue::Construct(const FArguments& InArgs)
+{
+	Path = InArgs._Path;
+	NotifyHook = InArgs._NotifyHook;
+
+	const FProperty* Property = Path.GetLastProperty();
+	const FEnumProperty* EnumProperty = CastField<const FEnumProperty>(Property);
+	const FNumericProperty* NumericProperty = CastField<const FNumericProperty>(Property);
+	const UEnum* Enum = EnumProperty ? EnumProperty->GetEnum() : NumericProperty ? NumericProperty->GetIntPropertyEnum() : nullptr;
+	EnumType = Enum;
+
+	if (Enum && Property->ArrayDim == 1)
+	{
+		ChildSlot
+		[
+			SNew(SComboButton)
+			.ComboButtonStyle(&::UE::AdvancedWidgets::FAdvancedWidgetsStyle::Get().GetWidgetStyle<FComboButtonStyle>("PropertyValue.ComboButton"))
+			.ContentPadding(FMargin(0.f))
+			.HAlign(HAlign_Left)
+			.OnGetMenuContent(this, &SEnumPropertyValue::OnGetMenuContent)
+			.ButtonContent()
+			[
+				SNew(STextBlock)
+				.Text(this, &SEnumPropertyValue::GetText)
+			]
+ 		];
+	}
+}
+
+
+FText SEnumPropertyValue::GetText() const
+{
+	if (const UEnum* EnumPtr = EnumType.Get())
+	{
+		int64 CurrentValue = GetCurrentValue();
+		return EnumPtr->GetDisplayNameTextByValue(CurrentValue);
+	}
+
+	return FText::GetEmpty();
+}
+
+
+int32 SEnumPropertyValue::GetCurrentValue() const
+{
+	if (const UEnum* EnumPtr = EnumType.Get())
+	{
+		int64 CurrentValue = INDEX_NONE;
+		if (const FNumericProperty* NumericProperty = CastField<const FNumericProperty>(Path.GetLastProperty()))
+		{
+			if (const void* Container = Path.GetContainerPtr())
+			{
+				CurrentValue = NumericProperty->GetSignedIntPropertyValue(NumericProperty->ContainerPtrToValuePtr<const void*>(Container));
+			}
+		}
+		else if (const FEnumProperty* EnumProperty = CastField<const FEnumProperty>(Path.GetLastProperty()))
+		{
+			if (const void* Container = Path.GetContainerPtr())
+			{
+				CurrentValue = EnumProperty->GetUnderlyingProperty()->GetSignedIntPropertyValue(EnumProperty->ContainerPtrToValuePtr<const void*>(Container));
+			}
+		}
+
+		return CurrentValue;
+	}
+	return 0;
+}
+
+
+TSharedRef<SWidget> SEnumPropertyValue::OnGetMenuContent()
+{
+	const bool bCloseAfterSelection = true;
+	FMenuBuilder MenuBuilder(bCloseAfterSelection, nullptr);
+
+	if (const UEnum* EnumPtr = EnumType.Get())
+	{
+		// bitflag is not supported at runtime (WITH_EDITORONLY_DATA)
+		const bool bHasMaxValue = EnumPtr->ContainsExistingMax();
+		const int32 NumEnums = bHasMaxValue ? EnumPtr->NumEnums() - 1 : EnumPtr->NumEnums();
+
+		for (int32 Index = 0; Index < NumEnums; ++Index)
+		{
+#if WITH_EDITORONLY_DATA
+			if (!EnumPtr->HasMetaData(TEXT("Hidden"), Index))
+#endif
+			{
+				MenuBuilder.AddMenuEntry(
+					EnumPtr->GetDisplayNameTextByIndex(Index),
+					EnumPtr->GetToolTipTextByIndex(Index),
+					FSlateIcon(),
+					FUIAction
+					(
+						FExecuteAction::CreateSP(this, &SEnumPropertyValue::SetEnumEntry, Index),
+						FCanExecuteAction(),
+						//FIsActionChecked::CreateSP(this, &SEnumPropertyValue::IsEnumEntryChecked, Index)
+						FIsActionChecked()
+					),
+					NAME_None,
+					EUserInterfaceActionType::None);
+			}
+		}
+	}
+
+	return MenuBuilder.MakeWidget();
+}
+
+
+void SEnumPropertyValue::SetEnumEntry(int32 Index)
+{
+	if (const UEnum* EnumPtr = EnumType.Get())
+	{
+		int64 NewValue = EnumPtr->GetValueByIndex(Index);
+		if (const FNumericProperty* NumericProperty = CastField<const FNumericProperty>(Path.GetLastProperty()))
+		{
+			if (void* Container = Path.GetContainerPtr())
+			{
+				if (NotifyHook)
+				{
+					NotifyHook->OnPreValueChange(Path);
+				}
+				NumericProperty->SetIntPropertyValue(NumericProperty->ContainerPtrToValuePtr<const void*>(Container), NewValue);
+				if (NotifyHook)
+				{
+					NotifyHook->OnPostValueChange(Path);
+				}
+			}
+		}
+		else if (const FEnumProperty* EnumProperty = CastField<const FEnumProperty>(Path.GetLastProperty()))
+		{
+			if (void* Container = Path.GetContainerPtr())
+			{
+				if (NotifyHook)
+				{
+					NotifyHook->OnPreValueChange(Path);
+				}
+				EnumProperty->GetUnderlyingProperty()->SetIntPropertyValue(EnumProperty->ContainerPtrToValuePtr<const void*>(Container), NewValue);
+				 if (NotifyHook)
+				 {
+					 NotifyHook->OnPostValueChange(Path);
+				 }
+			}
+		}
+	}
+}
+
+
+//bool SEnumPropertyValue::IsEnumEntryChecked(int32 Index) const
+//{
+//	if (UEnum* EnumPtr = EnumType.Get())
+//	{
+//		return (GetCurrentValue() & EnumPtr->GetValueByIndex(Index)) != 0;
+//	}
+//	return false;
+//}
+
+} //namespace
diff --git a/Engine/Source/Runtime/AdvancedWidgets/Private/Widgets/PropertyViewer/SFieldIcon.cpp b/Engine/Source/Runtime/AdvancedWidgets/Private/Widgets/PropertyViewer/SFieldIcon.cpp
new file mode 100644
--- /dev/null
+++ b/Engine/Source/Runtime/AdvancedWidgets/Private/Widgets/PropertyViewer/SFieldIcon.cpp
@@ -0,0 +1,85 @@
+// Copyright Epic Games, Inc. All Rights Reserved.
+
+#include "Widgets/PropertyViewer/SFieldIcon.h"
+
+#include "Framework/PropertyViewer/FieldIconFinder.h"
+#include "Widgets/Images/SLayeredImage.h"
+
+
+namespace UE::PropertyViewer
+{
+
+
+void SFieldIcon::Construct(const FArguments& InArgs, const UClass* Class)
+{
+	check(Class);
+
+	ChildSlot
+	[
+		SNew(SImage)
+		.Image(FFieldIconFinder::GetIcon(Class))
+	];
+}
+
+
+void SFieldIcon::Construct(const FArguments& InArgs, const UScriptStruct* Struct)
+{
+	check(Struct);
+
+	ChildSlot
+	[
+		SNew(SImage)
+		.Image(FFieldIconFinder::GetIcon(Struct))
+	];
+}
+
+
+void SFieldIcon::Construct(const FArguments& InArgs, const FProperty* Property)
+{
+	check(Property);
+
+	FFieldIconFinder::FFieldIconArray Icons = InArgs._OverrideColorSettings.IsSet() ? FFieldIconFinder::GetPropertyIcon(Property, InArgs._OverrideColorSettings.GetValue()) : FFieldIconFinder::GetPropertyIcon(Property);
+	if (Icons.Num() == 1)
+	{
+		ChildSlot
+		[
+			SNew(SImage)
+			.Image(Icons[0].Icon)
+			.ColorAndOpacity(Icons[0].Color)
+		];
+	}
+	else
+	{
+		TArray<SLayeredImage::ImageLayer> Layers;
+		Layers.Reserve(Icons.Num());
+		for (const FFieldIconFinder::FFieldIcon& Icon : Icons)
+		{
+			Layers.Emplace(Icon.Icon, Icon.Color);
+		}
+
+		ChildSlot
+		[
+			SNew(SLayeredImage, MoveTemp(Layers))
+		];
+	}
+}
+
+
+void SFieldIcon::Construct(const FArguments& InArgs, const UFunction* Function)
+{
+	check(Function);
+
+	FFieldIconFinder::FFieldIconArray Icons = InArgs._OverrideColorSettings.IsSet() ? FFieldIconFinder::GetFunctionIcon(Function, InArgs._OverrideColorSettings.GetValue()) : FFieldIconFinder::GetFunctionIcon(Function);
+	if (Icons.Num() >= 1)
+	{
+		ChildSlot
+		[
+			SNew(SImage)
+			.Image(Icons[0].Icon)
+			.ColorAndOpacity(Icons[0].Color)
+		];
+	}
+}
+
+
+} //namespace
diff --git a/Engine/Source/Runtime/AdvancedWidgets/Private/Widgets/PropertyViewer/SFieldName.cpp b/Engine/Source/Runtime/AdvancedWidgets/Private/Widgets/PropertyViewer/SFieldName.cpp
new file mode 100644
--- /dev/null
+++ b/Engine/Source/Runtime/AdvancedWidgets/Private/Widgets/PropertyViewer/SFieldName.cpp
@@ -0,0 +1,173 @@
+// Copyright Epic Games, Inc. All Rights Reserved.
+
+#include "Widgets/PropertyViewer/SFieldName.h"
+
+#include "Application/SlateApplicationBase.h"
+#include "UObject/UnrealType.h"
+#include "UObject/Class.h"
+#include "Widgets/PropertyViewer/SFieldIcon.h"
+#include "Widgets/Text/STextBlock.h"
+#include "Widgets/SBoxPanel.h"
+
+
+namespace UE::PropertyViewer
+{
+
+void SFieldName::Construct(const FArguments& InArgs, const UClass* Class)
+{
+	check(Class);
+
+	TSharedPtr<SWidget> Icon;
+	if (InArgs._bShowIcon)
+	{
+		Icon = SNew(SFieldIcon, Class);
+	}
+
+	FText DisplayName;
+	if (InArgs._OverrideDisplayName.IsSet())
+	{
+		DisplayName = InArgs._OverrideDisplayName.GetValue();
+	}
+	else
+	{
+		DisplayName = InArgs._bSanitizeName ? Class->GetDisplayNameText() : FText::FromName(Class->GetFName());
+	}
+
+	Field = Class;
+	Construct(InArgs, DisplayName, Icon);
+}
+
+
+void SFieldName::Construct(const FArguments& InArgs, const UScriptStruct* Struct)
+{
+	check(Struct);
+
+	TSharedPtr<SWidget> Icon;
+	if (InArgs._bShowIcon)
+	{
+		Icon = SNew(SFieldIcon, Struct);
+	}
+
+	FText DisplayName;
+	if (InArgs._OverrideDisplayName.IsSet())
+	{
+		DisplayName = InArgs._OverrideDisplayName.GetValue();
+	}
+	else
+	{
+		DisplayName = InArgs._bSanitizeName ? Struct->GetDisplayNameText() : FText::FromName(Struct->GetFName());
+	}
+
+	Field = Struct;
+	Construct(InArgs, DisplayName, Icon);
+}
+
+
+void SFieldName::Construct(const FArguments& InArgs, const FProperty* Property)
+{
+	check(Property);
+
+	TSharedPtr<SWidget> Icon;
+	if (InArgs._bShowIcon)
+	{
+		Icon = SNew(SFieldIcon, Property);
+	}
+	
+	FText DisplayName;
+	if (InArgs._OverrideDisplayName.IsSet())
+	{
+		DisplayName = InArgs._OverrideDisplayName.GetValue();
+	}
+	else
+	{
+		DisplayName = InArgs._bSanitizeName ? Property->GetDisplayNameText() : FText::FromName(Property->GetFName());
+	}
+
+	Field = Property;
+	Construct(InArgs, DisplayName, Icon);
+}
+
+
+void SFieldName::Construct(const FArguments& InArgs, const UFunction* Function)
+{
+	check(Function);
+
+	TSharedPtr<SWidget> Icon;
+	if (InArgs._bShowIcon)
+	{
+		Icon = SNew(SFieldIcon, Function);
+	}
+
+	FText DisplayName;
+	if (InArgs._OverrideDisplayName.IsSet())
+	{
+		DisplayName = InArgs._OverrideDisplayName.GetValue();
+	}
+	else
+	{
+		DisplayName = InArgs._bSanitizeName ? Function->GetDisplayNameText() : FText::FromName(Function->GetFName());
+	}
+
+	Field = Function;
+	Construct(InArgs, DisplayName, Icon);
+}
+
+
+void SFieldName::Construct(const FArguments& InArgs, const FText& DisplayName, TSharedPtr<SWidget> Icon)
+{
+	if (Icon)
+	{
+		ChildSlot
+		[
+			SNew(SHorizontalBox)
+			+ SHorizontalBox::Slot()
+			.AutoWidth()
+			.HAlign(HAlign_Right)
+			.VAlign(VAlign_Center)
+			[
+				Icon.ToSharedRef()
+			]
+			+ SHorizontalBox::Slot()
+			.Padding(4.0f)
+			[
+				SAssignNew(NameBlock, STextBlock)
+				.Text(DisplayName)
+				.HighlightText(InArgs._HighlightText)
+			]
+		];
+	}
+	else
+	{
+		ChildSlot
+		[
+			SAssignNew(NameBlock, STextBlock)
+			.Text(DisplayName)
+			.HighlightText(InArgs._HighlightText)
+		];
+	}
+	SetToolTip(TAttribute<TSharedPtr<IToolTip>>::CreateSP(this, &SFieldName::CreateToolTip));
+}
+
+
+void SFieldName::SetHighlightText(TAttribute<FText> InHighlightText)
+{
+	if (NameBlock)
+	{
+		NameBlock->SetHighlightText(MoveTemp(InHighlightText));
+	}
+}
+
+TSharedPtr<IToolTip> SFieldName::CreateToolTip() const
+{
+	if (FProperty* PropertyPtr = Field.Get<FProperty>())
+	{
+		return FSlateApplicationBase::Get().MakeToolTip(PropertyPtr->GetToolTipText());
+	}
+	if (UField* FieldPtr = Field.Get<UField>())
+	{
+		return FSlateApplicationBase::Get().MakeToolTip(FieldPtr->GetToolTipText());
+	}
+	return TSharedPtr<IToolTip>();
+}
+
+} //namespace
diff --git a/Engine/Source/Runtime/AdvancedWidgets/Private/Widgets/PropertyViewer/SNumericPropertyValue.cpp b/Engine/Source/Runtime/AdvancedWidgets/Private/Widgets/PropertyViewer/SNumericPropertyValue.cpp
new file mode 100644
--- /dev/null
+++ b/Engine/Source/Runtime/AdvancedWidgets/Private/Widgets/PropertyViewer/SNumericPropertyValue.cpp
@@ -0,0 +1,184 @@
+// Copyright Epic Games, Inc. All Rights Reserved.
+
+#include "Widgets/PropertyViewer/SNumericPropertyValue.h"
+
+#include "Framework/PropertyViewer/INotifyHook.h"
+#include "Styling/AdvancedWidgetsStyle.h"
+#include "Styling/SlateTypes.h"
+#include "UObject/UnrealType.h"
+
+#include "Widgets/Input/SSpinBox.h"
+
+#define LOCTEXT_NAMESPACE "SNumericPropertyValue"
+
+namespace UE::PropertyViewer
+{
+
+TSharedPtr<SWidget> SNumericPropertyValue::CreateInstance(const FPropertyValueFactory::FGenerateArgs Args)
+{
+	return SNew(SNumericPropertyValue)
+		.Path(Args.Path)
+		.NotifyHook(Args.NotifyHook)
+		.IsEnabled(Args.bCanEditValue);
+}
+
+void SNumericPropertyValue::Construct(const FArguments& InArgs)
+{
+	Path = InArgs._Path;
+	NotifyHook = InArgs._NotifyHook;
+
+	if (const FNumericProperty* Property = CastField<const FNumericProperty>(Path.GetLastProperty()))
+	{
+		if (Property->ArrayDim == 1)
+		{
+			const FSpinBoxStyle& SpinBoxStyle = ::UE::AdvancedWidgets::FAdvancedWidgetsStyle::Get().GetWidgetStyle<FSpinBoxStyle>("PropertyValue.SpinBox");
+			// int8/int16/int32/int64
+			if (Property->CanHoldValue(-1) && Property->IsInteger())
+			{
+				ChildSlot
+				[
+					SNew(SSpinBox<int64>)
+					.Style(&SpinBoxStyle)
+					.Value(this, &SNumericPropertyValue::GetCurrentValue_int64)
+					.OnValueChanged(this, &SNumericPropertyValue::OnValueChanged_int64)
+					.OnValueCommitted(this, &SNumericPropertyValue::OnValueCommitted_int64)
+					.OnEndSliderMovement(this, &SNumericPropertyValue::OnEndSliderMovement_int64)
+				];
+			}
+			// uint8/uint16/uint32/uint64
+			else if (Property->IsInteger())
+			{
+				ChildSlot
+				[
+					SNew(SSpinBox<uint64>)
+					.Style(&SpinBoxStyle)
+					.Value(this, &SNumericPropertyValue::GetCurrentValue_uint64)
+					.OnValueChanged(this, &SNumericPropertyValue::OnValueChanged_uint64)
+					.OnValueCommitted(this, &SNumericPropertyValue::OnValueCommitted_uint64)
+					.OnEndSliderMovement(this, &SNumericPropertyValue::OnEndSliderMovement_uint64)
+				];
+			}
+			// float/double
+			else
+			{
+				ChildSlot
+				[
+					SNew(SSpinBox<double>)
+					.Style(&SpinBoxStyle)
+					.Value(this, &SNumericPropertyValue::GetCurrentValue_double)
+					.OnValueChanged(this, &SNumericPropertyValue::OnValueChanged_double)
+					.OnValueCommitted(this, &SNumericPropertyValue::OnValueCommitted_double)
+					.OnEndSliderMovement(this, &SNumericPropertyValue::OnEndSliderMovement_double)
+				];
+			}
+		}
+	}
+}
+
+void SNumericPropertyValue::OnEndSliderMovement_uint64(uint64 Value)
+{
+	OnValueChanged_uint64(Value);
+}
+void SNumericPropertyValue::OnEndSliderMovement_int64(int64 Value)
+{
+	OnValueChanged_int64(Value);
+}
+void SNumericPropertyValue::OnEndSliderMovement_double(double Value)
+{
+	OnValueChanged_double(Value);
+}
+
+
+namespace Private
+{
+	template<typename TValue, typename TPred>
+	TValue GetCurrentValue(const FPropertyPath& Path, TPred Pred)
+	{
+		if (const void* Container = Path.GetContainerPtr())
+		{
+			if (const FNumericProperty* Property = CastField<const FNumericProperty>(Path.GetLastProperty()))
+			{
+				return Pred(Property, Property->ContainerPtrToValuePtr<const void*>(Container));
+			}
+		}
+		return TValue();
+	}
+}// namespace
+
+
+uint64 SNumericPropertyValue::GetCurrentValue_uint64() const
+{
+	return Private::GetCurrentValue<uint64>(Path, [](const FNumericProperty* Property, const void* DataPtr) { return Property->GetUnsignedIntPropertyValue(DataPtr); });
+}
+int64 SNumericPropertyValue::GetCurrentValue_int64() const
+{
+	return Private::GetCurrentValue<int64>(Path, [](const FNumericProperty* Property, const void* DataPtr) { return Property->GetSignedIntPropertyValue(DataPtr); });
+}
+double SNumericPropertyValue::GetCurrentValue_double() const
+{
+	return Private::GetCurrentValue<double>(Path, [](const FNumericProperty* Property, const void* DataPtr) { return Property->GetFloatingPointPropertyValue(DataPtr); });
+}
+
+
+namespace Private
+{
+	template<typename TPred>
+	void OnValueChanged(FPropertyPath& Path, INotifyHook* NotifyHook, TPred Pred)
+	{
+		if (void* Container = Path.GetContainerPtr())
+		{
+			if (const FNumericProperty* Property = CastField<const FNumericProperty>(Path.GetLastProperty()))
+			{
+				if (NotifyHook)
+				{
+					NotifyHook->OnPreValueChange(Path);
+				}
+				Pred(Property, Property->ContainerPtrToValuePtr<void*>(Container));
+				if (NotifyHook)
+				{
+					NotifyHook->OnPostValueChange(Path);
+				}
+			}
+		}
+	}
+}// namespace
+
+void SNumericPropertyValue::OnValueChanged_uint64(uint64 Value)
+{
+	Private::OnValueChanged(Path, NotifyHook, [Value](const FNumericProperty* Property, void* DataPtr) { Property->SetIntPropertyValue(DataPtr, Value); });
+}
+void SNumericPropertyValue::OnValueChanged_int64(int64 Value)
+{
+	Private::OnValueChanged(Path, NotifyHook, [Value](const FNumericProperty* Property, void* DataPtr) { Property->SetIntPropertyValue(DataPtr, Value); });
+}
+void SNumericPropertyValue::OnValueChanged_double(double Value)
+{
+	Private::OnValueChanged(Path, NotifyHook, [Value](const FNumericProperty* Property, void* DataPtr) { Property->SetFloatingPointPropertyValue(DataPtr, Value); });
+}
+
+
+void SNumericPropertyValue::OnValueCommitted_uint64(uint64 Value, ETextCommit::Type CommitInfo)
+{
+	if (CommitInfo == ETextCommit::OnEnter || CommitInfo == ETextCommit::OnUserMovedFocus)
+	{
+		OnValueChanged_uint64(Value);
+	}
+}
+void SNumericPropertyValue::OnValueCommitted_int64(int64 Value, ETextCommit::Type CommitInfo)
+{
+	if (CommitInfo == ETextCommit::OnEnter || CommitInfo == ETextCommit::OnUserMovedFocus)
+	{
+		OnValueChanged_int64(Value);
+	}
+}
+void SNumericPropertyValue::OnValueCommitted_double(double Value, ETextCommit::Type CommitInfo)
+{
+	if (CommitInfo == ETextCommit::OnEnter || CommitInfo == ETextCommit::OnUserMovedFocus)
+	{
+		OnValueChanged_double(Value);
+	}
+}
+
+} //namespace
+
+#undef LOCTEXT_NAMESPACE
diff --git a/Engine/Source/Runtime/AdvancedWidgets/Private/Widgets/PropertyViewer/SPropertyViewer.cpp b/Engine/Source/Runtime/AdvancedWidgets/Private/Widgets/PropertyViewer/SPropertyViewer.cpp
new file mode 100644
--- /dev/null
+++ b/Engine/Source/Runtime/AdvancedWidgets/Private/Widgets/PropertyViewer/SPropertyViewer.cpp
@@ -0,0 +1,115 @@
+// Copyright Epic Games, Inc. All Rights Reserved.
+
+#include "Widgets/PropertyViewer/SPropertyViewer.h"
+#include "Widgets/PropertyViewer/SPropertyViewerImpl.h"
+
+#define LOCTEXT_NAMESPACE "SPropertyViewer"
+
+namespace UE::PropertyViewer
+{
+
+void SPropertyViewer::ConstructInternal(const FArguments& InArgs)
+{
+	ChildSlot
+	[
+		Implementation->Construct(InArgs)
+	];
+}
+
+void SPropertyViewer::Construct(const FArguments& InArgs)
+{
+	Implementation = MakeShared<Private::FPropertyViewerImpl>(InArgs);
+	ConstructInternal(InArgs);
+}
+void SPropertyViewer::Construct(const FArguments& InArgs, const UScriptStruct* Struct)
+{
+	Implementation = MakeShared<Private::FPropertyViewerImpl>(InArgs);
+	Implementation->AddContainer(MakeContainerIdentifier(), Struct);
+	ConstructInternal(InArgs);
+}
+void SPropertyViewer::Construct(const FArguments& InArgs, const UScriptStruct* Struct, void* InData)
+{
+	Implementation = MakeShared<Private::FPropertyViewerImpl>(InArgs);
+	Implementation->AddContainerInstance(MakeContainerIdentifier(), Struct, InData);
+	ConstructInternal(InArgs);
+}
+void SPropertyViewer::Construct(const FArguments& InArgs, const UClass* Class)
+{
+	Implementation = MakeShared<Private::FPropertyViewerImpl>(InArgs);
+	Implementation->AddContainer(MakeContainerIdentifier(), Class);
+	ConstructInternal(InArgs);
+}
+void SPropertyViewer::Construct(const FArguments& InArgs, UObject* ObjectInstance)
+{
+	Implementation = MakeShared<Private::FPropertyViewerImpl>(InArgs);
+	Implementation->AddContainerInstance(MakeContainerIdentifier(), ObjectInstance);
+	ConstructInternal(InArgs);
+}
+void SPropertyViewer::Construct(const FArguments& InArgs, const UFunction* Function)
+{
+	Implementation = MakeShared<Private::FPropertyViewerImpl>(InArgs);
+	Implementation->AddContainer(MakeContainerIdentifier(), Function);
+	ConstructInternal(InArgs);
+}
+
+
+SPropertyViewer::FHandle SPropertyViewer::AddContainer(const UScriptStruct* Struct)
+{
+	SPropertyViewer::FHandle Result = MakeContainerIdentifier();
+	Implementation->AddContainer(Result, Struct);
+	return Result;
+}
+SPropertyViewer::FHandle SPropertyViewer::AddContainer(const UClass* Class)
+{
+	SPropertyViewer::FHandle Result = MakeContainerIdentifier();
+	Implementation->AddContainer(Result, Class);
+	return Result;
+}
+SPropertyViewer::FHandle SPropertyViewer::AddContainer(const UFunction* Function)
+{
+	SPropertyViewer::FHandle Result = MakeContainerIdentifier();
+	Implementation->AddContainer(Result, Function);
+	return Result;
+}
+
+
+SPropertyViewer::FHandle SPropertyViewer::AddInstance(const UScriptStruct* Struct, void* InData)
+{
+	check(InData);
+	SPropertyViewer::FHandle Result = MakeContainerIdentifier();
+	Implementation->AddContainerInstance(Result, Struct, InData);
+	return Result;
+}
+SPropertyViewer::FHandle SPropertyViewer::AddInstance(UObject* ObjectInstance)
+{
+	SPropertyViewer::FHandle Result = MakeContainerIdentifier();
+	Implementation->AddContainerInstance(Result, ObjectInstance);
+	return Result;
+}
+
+
+void SPropertyViewer::Remove(FHandle Identifier)
+{
+	Implementation->Remove(Identifier);
+}
+
+
+void SPropertyViewer::RemoveAll()
+{
+
+	Implementation->RemoveAll();
+}
+
+
+SPropertyViewer::FHandle SPropertyViewer::MakeContainerIdentifier()
+{
+	static int32 IdentifierGenerator = 0;
+	SPropertyViewer::FHandle Result;
+	++IdentifierGenerator;
+	Result.Id = IdentifierGenerator;
+	return Result;
+}
+
+} //namespace
+
+#undef LOCTEXT_NAMESPACE
diff --git a/Engine/Source/Runtime/AdvancedWidgets/Private/Widgets/PropertyViewer/SPropertyViewerImpl.cpp b/Engine/Source/Runtime/AdvancedWidgets/Private/Widgets/PropertyViewer/SPropertyViewerImpl.cpp
new file mode 100644
--- /dev/null
+++ b/Engine/Source/Runtime/AdvancedWidgets/Private/Widgets/PropertyViewer/SPropertyViewerImpl.cpp
@@ -0,0 +1,1005 @@
+// Copyright Epic Games, Inc. All Rights Reserved.
+
+#include "Widgets/PropertyViewer/SPropertyViewerImpl.h"
+
+#include "AdvancedWidgetsModule.h"
+#include "Misc/IFilter.h"
+#include "Misc/TextFilter.h"
+#include "Framework/PropertyViewer/IFieldExpander.h"
+#include "Framework/PropertyViewer/IFieldIterator.h"
+#include "Framework/PropertyViewer/INotifyHook.h"
+#include "Framework/PropertyViewer/PropertyPath.h"
+#include "Framework/PropertyViewer/PropertyValueFactory.h"
+#include "Framework/Views/TreeFilterHandler.h"
+
+#include "Styling/SlateIconFinder.h"
+
+#include "Widgets/Input/SSpinBox.h"
+#include "Widgets/Input/SSearchBox.h"
+#include "Widgets/Layout/SBorder.h"
+#include "Widgets/PropertyViewer/SFieldName.h"
+#include "Widgets/Text/SInlineEditableTextBlock.h"
+#include "Widgets/Views/SHeaderRow.h"
+#include "Widgets/Views/STreeView.h"
+#include "Widgets/SBoxPanel.h"
+#include "Widgets/SNullWidget.h"
+
+#include "UObject/UnrealType.h"
+
+#if WITH_EDITOR
+#include "Editor.h"
+#include "Editor/EditorEngine.h"
+#endif //WITH_EDITOR
+
+#define LOCTEXT_NAMESPACE "SPropertyViewerImpl"
+
+namespace UE::PropertyViewer::Private
+{
+
+/** 
+ * Column name 
+ */
+static FName ColumnName_FieldPreWidget = "FieldPreWidget";
+static FName ColumnName_Field = "Field";
+static FName ColumnName_PropertyValue = "FieldValue";
+static FName ColumnName_FieldPostWidget = "FieldPostWidget";
+
+
+/**
+ * FContainer
+ */
+FContainer::FContainer(SPropertyViewer::FHandle InIdentifier, UObject* InstanceToDisplay)
+	: Identifier(InIdentifier)
+	, Container(InstanceToDisplay->GetClass())
+	, ObjectInstance(InstanceToDisplay)
+	, bIsObject(true)
+{
+}
+
+
+FContainer::FContainer(SPropertyViewer::FHandle InIdentifier, const UStruct* ClassToDisplay)
+	: Identifier(InIdentifier)
+	, Container(ClassToDisplay)
+{
+}
+
+
+FContainer::FContainer(SPropertyViewer::FHandle InIdentifier, const UScriptStruct* Struct, void* Data)
+	: Identifier(InIdentifier)
+	, Container(Struct)
+	, StructInstance(Data)
+{
+}
+
+
+bool FContainer::IsValid() const
+{
+	if (const UClass* Class = Cast<UClass>(Container.Get()))
+	{
+		return !Class->HasAnyClassFlags(EClassFlags::CLASS_NewerVersionExists)
+			&& (!bIsObject || (ObjectInstance.Get() && ObjectInstance.Get()->GetClass() == Class));
+	}
+	else if (const UScriptStruct* ScriptStruct = Cast<UScriptStruct>(Container.Get()))
+	{
+		return (ScriptStruct->StructFlags & (EStructFlags::STRUCT_Trashed)) != 0;
+	}
+	return false;
+}
+
+
+/**
+ * FTreeNode
+ */
+TSharedRef<FTreeNode> FTreeNode::MakeContainer(const TSharedPtr<FContainer>& InContainer, TOptional<FText> InDisplayName)
+{
+	TSharedRef<FTreeNode> Result = MakeShared<FTreeNode>();
+	Result->Container = InContainer;
+	Result->OverrideDisplayName = InDisplayName;
+	return Result;
+}
+
+
+TSharedRef<FTreeNode> FTreeNode::MakeField(TSharedPtr<FTreeNode> InParent, const FProperty* Property, TOptional<FText> InDisplayName)
+{
+	check(Property);
+	TSharedRef<FTreeNode> Result = MakeShared<FTreeNode>();
+	Result->Property = Property;
+	Result->OverrideDisplayName = InDisplayName;
+	Result->ParentNode = InParent;
+	InParent->ChildNodes.Add(Result);
+	return Result;
+}
+
+
+TSharedRef<FTreeNode> FTreeNode::MakeField(TSharedPtr<FTreeNode> InParent, const UFunction* Function, TOptional<FText> InDisplayName)
+{
+	check(Function);
+	TSharedRef<FTreeNode> Result = MakeShared<FTreeNode>();
+	Result->Function = Function;
+	Result->OverrideDisplayName = InDisplayName;
+	Result->ParentNode = InParent;
+	InParent->ChildNodes.Add(Result);
+	return Result;
+}
+
+
+FPropertyPath FTreeNode::GetPropertyPath() const
+{
+	const FTreeNode* CurrentNode = this;
+	FPropertyPath::FPropertyArray Properties;
+	while (CurrentNode)
+	{
+		if (CurrentNode->Property)
+		{
+			Properties.Insert(Property, 0);
+		}
+
+		TSharedPtr<FContainer> ContainerPin = CurrentNode->Container.Pin();
+		if (ContainerPin)
+		{
+			if (ContainerPin->IsObjectInstance())
+			{
+				if (UObject* ObjectInstance = ContainerPin->GetObjectInstance())
+				{
+					return FPropertyPath(ObjectInstance, MoveTemp(Properties));
+				}
+			}
+			else if (ContainerPin->IsScriptStructInstance())
+			{
+				if (const UStruct* ContainerStruct = ContainerPin->GetStruct())
+				{
+					return FPropertyPath(CastChecked<const UScriptStruct>(ContainerStruct), ContainerPin->GetScriptStructInstance(), MoveTemp(Properties));
+				}
+			}
+			else if (const UClass* ContainerClass = Cast<const UClass>(ContainerPin->GetStruct()))
+			{
+				return FPropertyPath(ContainerClass->GetDefaultObject(), MoveTemp(Properties));
+			}
+			return FPropertyPath();
+		}
+
+		if (TSharedPtr<FTreeNode> ParentNodePin = CurrentNode->ParentNode.Pin())
+		{
+			CurrentNode = ParentNodePin.Get();
+		}
+		else if (ContainerPin == nullptr)
+		{
+			ensureMsgf(false, TEXT("The tree is not owned by a container"));
+			return FPropertyPath();
+		}
+	}
+	return FPropertyPath();
+}
+
+
+TSharedPtr<FContainer> FTreeNode::GetOwnerContainer() const
+{
+	TSharedPtr<FContainer> Result;
+	const FTreeNode* CurrentNode = this;
+	while(CurrentNode)
+	{
+		if (TSharedPtr<FContainer> ContainerPin = CurrentNode->Container.Pin())
+		{
+			return ContainerPin;
+		}
+		if (TSharedPtr<FTreeNode> ParentNodePin = CurrentNode->ParentNode.Pin())
+		{
+			CurrentNode = ParentNodePin.Get();
+		}
+		else
+		{
+			ensureMsgf(false, TEXT("The tree is not owned by a container"));
+			break;
+		}
+	}
+	return TSharedPtr<FContainer>();
+}
+
+
+void FTreeNode::GetFilterStrings(TArray<FString>& OutStrings) const
+{
+	if (Property)
+	{
+		OutStrings.Add(Property->GetName());
+		OutStrings.Add(Property->GetDisplayNameText().ToString());
+	}
+	if (const UFunction* FunctionPtr = Function.Get())
+	{
+		OutStrings.Add(FunctionPtr->GetName());
+		OutStrings.Add(FunctionPtr->GetDisplayNameText().ToString());
+	}
+	if (const TSharedPtr<FContainer> ContainerPtr = Container.Pin())
+	{
+		if (const UStruct* StructPtr = ContainerPtr->GetStruct())
+		{
+			OutStrings.Add(StructPtr->GetName());
+			OutStrings.Add(StructPtr->GetDisplayNameText().ToString());
+		}
+	}
+	if (GetOverrideDisplayName().IsSet())
+	{
+		OutStrings.Add(GetOverrideDisplayName().GetValue().ToString());
+	}
+}
+
+
+void FTreeNode::BuildChildNodes(IFieldIterator& FieldIterator, IFieldExpander& FieldExpander)
+{
+	BuildChildNodesRecursive(FieldIterator, FieldExpander, 2);
+}
+
+
+void FTreeNode::BuildChildNodesRecursive(IFieldIterator& FieldIterator, IFieldExpander& FieldExpander, int32 RecursiveCount)
+{
+	if (RecursiveCount <= 0)
+	{
+		return;
+	}
+	--RecursiveCount;
+
+	ChildNodes.Reset();
+
+	const UStruct* ChildStructType = nullptr;
+	if (Property)
+	{
+		if (const FStructProperty* StructProperty = CastField<const FStructProperty>(Property))
+		{
+			ChildStructType = StructProperty->Struct;
+		}
+		else if (const FObjectPropertyBase* ObjectProperty = CastField<const FObjectPropertyBase>(Property))
+		{
+			// If the container is an instance and the object is nullptr, do not expand this object.
+			bool bIsNullptr = false;
+			if (TSharedPtr<FContainer> OwnerContainer = GetOwnerContainer())
+			{
+				if (OwnerContainer->IsInstance())
+				{
+					const FPropertyPath PropertyPath = GetPropertyPath();
+					if (const void* ContainerPtr = PropertyPath.GetContainerPtr())
+					{
+						bIsNullptr = ObjectProperty->GetObjectPropertyValue_InContainer(ContainerPtr) == nullptr;
+					}
+				}
+			}
+
+			if (!bIsNullptr && FieldExpander.CanExpandObject(ObjectProperty->PropertyClass))
+			{
+				//todo expand the object instance, not the object class
+				//todo we want to confirm that that thing didn't change frame to frame
+				ChildStructType = ObjectProperty->PropertyClass;
+			}
+		}
+	}
+	else if (const UFunction* FunctionPtr = Function.Get())
+	{
+		if (FieldExpander.CanExpandFunction(FunctionPtr))
+		{
+			ChildStructType = FunctionPtr;
+		}
+	}
+	else if (const TSharedPtr<FContainer> ContainerPin = Container.Pin())
+	{
+		ChildStructType = ContainerPin->GetStruct();
+	}
+
+	if (ChildStructType)
+	{
+		for (const FFieldVariant FieldIt : FieldIterator.GetFields(ChildStructType))
+		{
+			if (const FProperty* PropertyIt = FieldIt.Get<FProperty>())
+			{
+				TSharedPtr<FTreeNode> Node = MakeField(AsShared(), PropertyIt, TOptional<FText>());
+				Node->BuildChildNodesRecursive(FieldIterator, FieldExpander, RecursiveCount);
+			}
+			if (const UFunction* FunctionIt = FieldIt.Get<UFunction>())
+			{
+				TSharedPtr<FTreeNode> Node = MakeField(AsShared(), FunctionIt, TOptional<FText>());
+				Node->BuildChildNodesRecursive(FieldIterator, FieldExpander, RecursiveCount);
+			}
+		}
+
+		ChildNodes.Sort(Sort);
+	}
+
+	bChildGenerated = true;
+}
+
+
+bool FTreeNode::Sort(const TSharedPtr<FTreeNode>& NodeA, const TSharedPtr<FTreeNode>& NodeB)
+{
+	bool bIsContainerA = NodeA->IsContainer();
+	bool bIsContainerB = NodeB->IsContainer();
+	bool bIsObjectPropertyA = CastField<FObjectPropertyBase>(NodeA->Property) != nullptr;
+	bool bIsObjectPropertyB = CastField<FObjectPropertyBase>(NodeB->Property) != nullptr;
+	bool bIsFunctionA = NodeA->Function.Get() != nullptr;
+	bool bIsFunctionB = NodeB->Function.Get() != nullptr;
+	const FName NodeStrA = bIsContainerA ? NodeA->GetContainer()->GetStruct()->GetFName() : NodeA->GetField().GetFName();
+	const FName NodeStrB = bIsContainerB ? NodeB->GetContainer()->GetStruct()->GetFName() : NodeB->GetField().GetFName();
+
+	if (bIsFunctionA && bIsFunctionB)
+	{
+		return NodeStrA.LexicalLess(NodeStrB);
+	}
+	if (bIsObjectPropertyA && bIsObjectPropertyB)
+	{
+		return NodeStrA.LexicalLess(NodeStrB);
+	}
+
+	if (bIsFunctionA)
+	{
+		return true;
+	}
+	if (bIsFunctionB)
+	{
+		return false;
+	}
+	if (bIsObjectPropertyA)
+	{
+		return true;
+	}
+	if (bIsObjectPropertyB)
+	{
+		return false;
+	}
+
+	return NodeStrA.LexicalLess(NodeStrB);
+};
+
+
+/**
+ * FPropertyViewerImpl
+ */
+FPropertyViewerImpl::FPropertyViewerImpl(const SPropertyViewer::FArguments& InArgs)
+{
+	FieldIterator = InArgs._FieldIterator;
+	FieldExpander = InArgs._FieldExpander;
+	if (FieldIterator == nullptr)
+	{
+		FieldIterator = new FFieldIterator_BlueprintVisible();
+		bOwnFieldIterator = true;
+	}
+	if (FieldExpander == nullptr)
+	{
+		FieldExpander = new FFieldExpander_NoExpand();
+		bOwnFieldExpander = true;
+	}
+
+	NotifyHook = InArgs._NotifyHook;
+	OnGetPreSlot = InArgs._OnGetPreSlot;
+	OnGetPostSlot = InArgs._OnGetPostSlot;
+	OnSelectionChanged = InArgs._OnSelectionChanged;
+	PropertyVisibility = InArgs._PropertyVisibility;
+	bSanitizeName = InArgs._bSanitizeName;
+	bShowFieldIcon = InArgs._bShowFieldIcon;
+
+#if WITH_EDITOR
+	if (GEditor)
+	{
+		GEditor->OnBlueprintCompiled().AddRaw(this, &FPropertyViewerImpl::HandleBlueprintCompiled);
+	}
+#endif
+}
+
+
+FPropertyViewerImpl::~FPropertyViewerImpl()
+{
+	if (bOwnFieldIterator)
+	{
+		delete FieldIterator;
+	}
+	if (bOwnFieldExpander)
+	{
+		delete FieldExpander;
+	}
+#if WITH_EDITOR
+	if (GEditor)
+	{
+		GEditor->OnBlueprintCompiled().RemoveAll(this);
+	}
+#endif
+}
+
+
+TSharedRef<SWidget> FPropertyViewerImpl::Construct(const SPropertyViewer::FArguments& InArgs)
+{
+	TSharedPtr<SHorizontalBox> SearchBox;
+	if (InArgs._bShowSearchBox || InArgs._SearchBoxPreSlot.Widget != SNullWidget::NullWidget || InArgs._SearchBoxPostSlot.Widget != SNullWidget::NullWidget)
+	{
+		SearchBox = SNew(SHorizontalBox);
+		if (InArgs._SearchBoxPreSlot.Widget != SNullWidget::NullWidget)
+		{
+			SearchBox->AddSlot()
+				.AutoWidth()
+				.HAlign(HAlign_Left)
+				.VAlign(VAlign_Center)
+				[
+					InArgs._SearchBoxPreSlot.Widget
+				];
+		}
+		if (InArgs._bShowSearchBox)
+		{
+			SearchBox->AddSlot()
+				.FillWidth(1.f)
+				.VAlign(VAlign_Center)
+				[
+					CreateSearch()
+				];
+		}
+		else
+		{
+			SearchBox->AddSlot()
+				.FillWidth(1.f)
+				[
+					SNullWidget::NullWidget
+				];
+		}
+		if (InArgs._SearchBoxPostSlot.Widget != SNullWidget::NullWidget)
+		{
+			SearchBox->AddSlot()
+				.AutoWidth()
+				.HAlign(HAlign_Right)
+				.VAlign(VAlign_Center)
+				[
+					InArgs._SearchBoxPostSlot.Widget
+				];
+		}
+	}
+
+	TSharedRef<SWidget> ConstructedTree = SNew(SBorder)
+		.BorderImage(FAppStyle::Get().GetBrush("Brushes.Recessed"))
+		.Padding(0)
+		[
+			CreateTree(OnGetPreSlot.IsBound(), PropertyVisibility != SPropertyViewer::EPropertyVisibility::Hidden, OnGetPostSlot.IsBound())
+		];
+
+	if (SearchBox)
+	{
+		return SNew(SVerticalBox)
+			+ SVerticalBox::Slot()
+			.Padding(4)
+			.AutoHeight()
+			[
+				SearchBox.ToSharedRef()
+			]
+			+ SVerticalBox::Slot()
+			.FillHeight(1.0f)
+			[
+				ConstructedTree
+			];
+	}
+	else
+	{
+		return ConstructedTree;
+	}
+}
+
+
+void FPropertyViewerImpl::AddContainer(SPropertyViewer::FHandle Identifier, const UStruct* Struct)
+{
+	TSharedPtr<FContainer> NewContainer = MakeShared<FContainer>(Identifier, Struct);
+	Containers.Add(NewContainer);
+
+	AddContainerInternal(Identifier, NewContainer);
+}
+
+
+void FPropertyViewerImpl::AddContainerInstance(SPropertyViewer::FHandle Identifier, UObject* Object)
+{
+	TSharedPtr<FContainer> NewContainer = MakeShared<FContainer>(Identifier, Object);
+	Containers.Add(NewContainer);
+
+	AddContainerInternal(Identifier, NewContainer);
+}
+
+
+void FPropertyViewerImpl::AddContainerInstance(SPropertyViewer::FHandle Identifier, const UScriptStruct* Struct, void* Data)
+{
+	TSharedPtr<FContainer> NewContainer = MakeShared<FContainer>(Identifier, Struct, Data);
+	Containers.Add(NewContainer);
+
+	AddContainerInternal(Identifier, NewContainer);
+}
+
+
+void FPropertyViewerImpl::AddContainerInternal(SPropertyViewer::FHandle Identifier, TSharedPtr<FContainer>& NewContainer)
+{
+	TSharedPtr<FTreeNode> NewNode = FTreeNode::MakeContainer(NewContainer, TOptional<FText>());
+	NewNode->BuildChildNodes(*FieldIterator, *FieldExpander);
+	TreeSource.Add(NewNode);
+
+	if (TreeWidget)
+	{
+		TreeWidget->SetItemExpansion(NewNode, true);
+	}
+	if (FilterHandler)
+	{
+		FilterHandler->RefreshAndFilterTree();
+	}
+}
+
+
+void FPropertyViewerImpl::Remove(SPropertyViewer::FHandle Identifier)
+{
+	bool bRemoved = false;
+	for (int32 Index = 0; Index < TreeSource.Num(); ++Index)
+	{
+		if (TSharedPtr<FContainer> Container = TreeSource[Index]->GetContainer())
+		{
+			if (Container->GetIdentifier() == Identifier)
+			{
+				TreeSource.RemoveAt(Index);
+				bRemoved = true;
+				break;
+			}
+		}
+	}
+
+	for (int32 Index = 0; Index < Containers.Num(); ++Index)
+	{
+		if (Containers[Index]->GetIdentifier() == Identifier)
+		{
+			Containers.RemoveAt(Index);
+			break;
+		}
+	}
+
+	if (bRemoved)
+	{
+		if (FilterHandler)
+		{
+			FilterHandler->RefreshAndFilterTree();
+		}
+		else
+		{
+			TreeWidget->RequestTreeRefresh();
+		}
+	}
+}
+
+
+void FPropertyViewerImpl::RemoveAll()
+{
+	bool bRemoved = TreeSource.Num() > 0 || Containers.Num() > 0;
+	TreeSource.Reset();
+	Containers.Reset();
+
+	if (bRemoved)
+	{
+		if (FilterHandler)
+		{
+			FilterHandler->RefreshAndFilterTree();
+		}
+		else
+		{
+			TreeWidget->RequestTreeRefresh();
+		}
+	}
+}
+
+
+TSharedRef<SWidget> FPropertyViewerImpl::CreateSearch()
+{
+	SearchFilter = MakeShared<FTextFilter>(FTextFilter::FItemToStringArray::CreateSP(this, &FPropertyViewerImpl::HandleGetFilterStrings));
+
+	FilterHandler = MakeShared<FTreeFilter>();
+	FilterHandler->SetFilter(SearchFilter.Get());
+	FilterHandler->SetRootItems(&TreeSource, &FilteredTreeSource);
+	FilterHandler->SetGetChildrenDelegate(FTreeFilter::FOnGetChildren::CreateRaw(this, &FPropertyViewerImpl::HandleGetChildren));
+
+	return SAssignNew(SearchBoxWidget, SSearchBox)
+		.HintText(LOCTEXT("SearchHintText", "Search"))
+		.OnTextChanged(this, &FPropertyViewerImpl::HandleSearchChanged);
+}
+
+
+FText FPropertyViewerImpl::SetRawFilterText(const FText& InFilterText)
+{
+	const bool bNewFilteringEnabled = !InFilterText.IsEmpty();
+	FilterHandler->SetIsEnabled(bNewFilteringEnabled);
+	SearchFilter->SetRawFilterText(InFilterText);
+	FilterHandler->RefreshAndFilterTree();
+
+	for (const TSharedPtr<FTreeNode>& Node : FilteredTreeSource)
+	{
+		SetHighlightTextRecursive(Node, InFilterText);
+	}
+
+	return SearchFilter->GetFilterErrorText();
+}
+
+
+TSharedRef<SWidget> FPropertyViewerImpl::CreateTree(bool bHasPreWidget, bool bShowPropertyValue, bool bHasPostWidget)
+{
+	TSharedPtr<SHeaderRow> HeaderRowWidget;
+	if (bHasPreWidget || bShowPropertyValue || bHasPostWidget)
+	{
+		bUseRows = true;
+
+		HeaderRowWidget = SNew(SHeaderRow)
+		.Visibility(EVisibility::Collapsed);
+
+		if (bHasPreWidget)
+		{
+			HeaderRowWidget->AddColumn(
+				SHeaderRow::FColumn::FArguments()
+				.ColumnId(ColumnName_FieldPreWidget)
+				.DefaultLabel(LOCTEXT("PropertyPreWidget", ""))
+			);
+		}
+
+		HeaderRowWidget->AddColumn(
+			SHeaderRow::FColumn::FArguments()
+			.ColumnId(ColumnName_Field)
+			.DefaultLabel(LOCTEXT("FieldName", "Field Name"))
+			.FillWidth(0.75f)
+		);
+
+		if (bShowPropertyValue)
+		{
+			HeaderRowWidget->AddColumn(
+				SHeaderRow::FColumn::FArguments()
+				.ColumnId(ColumnName_PropertyValue)
+				.FillSized(100)
+				.DefaultLabel(LOCTEXT("PropertyValue", "Field Value"))
+				.FillWidth(0.25f)
+			);
+		}
+		if (bHasPostWidget)
+		{
+			HeaderRowWidget->AddColumn(
+				SHeaderRow::FColumn::FArguments()
+				.ColumnId(ColumnName_FieldPostWidget)
+				.DefaultLabel(LOCTEXT("PropertyPostWidget", ""))
+			);
+		}
+	}
+
+	if (FilterHandler)
+	{
+		SAssignNew(TreeWidget, STreeView<TSharedPtr<FTreeNode>>)
+			.ItemHeight(1.0f)
+			.TreeItemsSource(&FilteredTreeSource)
+			.SelectionMode(ESelectionMode::Single)
+			.OnGetChildren(FilterHandler.ToSharedRef(), &FTreeFilter::OnGetFilteredChildren)
+			.OnGenerateRow(this, &FPropertyViewerImpl::HandleGenerateRow)
+			.OnSelectionChanged(this, &FPropertyViewerImpl::HandleSelectionChanged)
+			.OnContextMenuOpening(this, &FPropertyViewerImpl::HandleContextMenuOpening)
+			.HeaderRow(HeaderRowWidget);
+		FilterHandler->SetTreeView(TreeWidget.Get());
+	}
+	else
+	{
+		SAssignNew(TreeWidget, STreeView<TSharedPtr<FTreeNode>>)
+			.ItemHeight(1.0f)
+			.TreeItemsSource(&TreeSource)
+			.SelectionMode(ESelectionMode::Single)
+			.OnGetChildren(this, &FPropertyViewerImpl::HandleGetChildren)
+			.OnGenerateRow(this, &FPropertyViewerImpl::HandleGenerateRow)
+			.OnSelectionChanged(this, &FPropertyViewerImpl::HandleSelectionChanged)
+			.OnContextMenuOpening(this, &FPropertyViewerImpl::HandleContextMenuOpening)
+			.HeaderRow(HeaderRowWidget);
+	}
+
+	return TreeWidget.ToSharedRef();
+}
+
+
+void FPropertyViewerImpl::SetHighlightTextRecursive(const TSharedPtr<FTreeNode>& OwnerNode, const FText& HighlightText)
+{
+	if (TSharedPtr<SFieldName> PropertyNameWidget = OwnerNode->PropertyWidget.Pin())
+	{
+		PropertyNameWidget->SetHighlightText(HighlightText);
+	}
+
+	if (OwnerNode->bChildGenerated)
+	{
+		for (const TSharedPtr<FTreeNode>& Node : OwnerNode->ChildNodes)
+		{
+			SetHighlightTextRecursive(Node, HighlightText);
+		}
+	}
+}
+
+
+void FPropertyViewerImpl::HandleGetFilterStrings(TSharedPtr<FTreeNode> Item, TArray<FString>& OutStrings)
+{
+	Item->GetFilterStrings(OutStrings);
+}
+
+
+TSharedRef<ITableRow> FPropertyViewerImpl::HandleGenerateRow(TSharedPtr<FTreeNode> Item, const TSharedRef<STableViewBase>& OwnerTable)
+{
+	FText HighlightText = SearchFilter ? SearchFilter->GetRawFilterText() : FText::GetEmpty();
+
+	TSharedPtr<SWidget> ItemWidget;
+	if (TSharedPtr<FContainer> ContainerPin = Item->GetContainer())
+	{
+		if (ContainerPin->IsValid())
+		{
+			if (const UClass* Class = Cast<const UClass>(ContainerPin->GetStruct()))
+			{
+				ItemWidget = SNew(SFieldName, Class)
+					.bShowIcon(true)
+					.bSanitizeName(bSanitizeName)
+					.OverrideDisplayName(Item->GetOverrideDisplayName());
+			}
+			if (const UScriptStruct* Struct = Cast<const UScriptStruct>(ContainerPin->GetStruct()))
+			{
+				ItemWidget = SNew(SFieldName, Struct)
+					.bShowIcon(true)
+					.bSanitizeName(bSanitizeName)
+					.OverrideDisplayName(Item->GetOverrideDisplayName());
+			}
+		}
+	}
+	else if (TSharedPtr<FContainer> OwnerContainerPin = Item->GetOwnerContainer())
+	{
+		if (OwnerContainerPin->IsValid())
+		{
+			if (FFieldVariant FieldVariant = Item->GetField())
+			{
+				if (FProperty* Property = FieldVariant.Get<FProperty>())
+				{
+					TSharedPtr<SFieldName> FieldName = SNew(SFieldName, Property)
+						.bShowIcon(bShowFieldIcon)
+						.bSanitizeName(bSanitizeName)
+						.OverrideDisplayName(Item->GetOverrideDisplayName())
+						.HighlightText(HighlightText);
+					Item->PropertyWidget = FieldName;
+					ItemWidget = FieldName;
+				}
+				else if (UFunction* Function = FieldVariant.Get<UFunction>())
+				{
+					TSharedPtr<SFieldName> FieldName = SNew(SFieldName, Function)
+						.bShowIcon(bShowFieldIcon)
+						.bSanitizeName(bSanitizeName)
+						.OverrideDisplayName(Item->GetOverrideDisplayName())
+						.HighlightText(HighlightText);
+					Item->PropertyWidget = FieldName;
+					ItemWidget = FieldName;
+				}
+			}
+		}
+	}
+
+	struct SMultiRowType : public SMultiColumnTableRow<TSharedPtr<FTreeNode>>
+	{
+		void Construct(const FArguments& Args, const TSharedRef<FPropertyViewerImpl> PropertyViewer, const TSharedRef<STableViewBase>& OwnerTableView, TSharedRef<FTreeNode> InItem, TSharedRef<SWidget> InFieldWidget)
+		{
+			PropertyViewOwner = PropertyViewer;
+			Item = InItem;
+			FieldWidget = InFieldWidget;
+			SMultiColumnTableRow<TSharedPtr<FTreeNode>>::Construct(Args, OwnerTableView);
+		}
+
+		virtual TSharedRef<SWidget> GenerateWidgetForColumn(const FName& ColumnName) override
+		{
+			if (ColumnName == ColumnName_Field)
+			{
+				return SNew(SHorizontalBox)
+				+ SHorizontalBox::Slot()
+				.AutoWidth()
+				[
+					SNew(SExpanderArrow, SharedThis(this))
+					.IndentAmount(16)
+					.ShouldDrawWires(true)
+				]
+				+SHorizontalBox::Slot()
+				.AutoWidth()
+				.Padding(2.0f, 0.0f)
+				.VAlign(VAlign_Center)
+				[
+					FieldWidget.ToSharedRef()
+				];
+			}
+
+			if (ColumnName == ColumnName_PropertyValue)
+			{
+				TSharedPtr<FTreeNode> ItemPin = Item.Pin();
+				TSharedPtr<FPropertyViewerImpl> PropertyViewOwnerPin = PropertyViewOwner.Pin();
+				if (ItemPin && PropertyViewOwnerPin)
+				{
+					if (ItemPin->IsField())
+					{
+						FFieldVariant Field = ItemPin->GetField();
+						if (!Field.IsUObject())
+						{
+							INotifyHook* NotifyHook = nullptr;
+							if (const TSharedPtr<FPropertyViewerImpl> PropertyViewer = PropertyViewOwner.Pin())
+							{
+								NotifyHook = PropertyViewer->NotifyHook;
+							}
+
+							bool bCanEditContainer = false;
+							if (const TSharedPtr<FContainer> OwnerContainer = ItemPin->GetOwnerContainer())
+							{
+								bCanEditContainer = OwnerContainer->CanEdit();
+							}
+
+							FPropertyValueFactory::FGenerateArgs Args;
+							Args.Path = ItemPin->GetPropertyPath();
+							Args.NotifyHook = NotifyHook;
+							Args.bCanEditValue = bCanEditContainer
+								&& PropertyViewOwnerPin->PropertyVisibility == SPropertyViewer::EPropertyVisibility::Editable
+								&& Args.Path.GetLastProperty() != nullptr
+								&& !Args.Path.GetLastProperty()->HasAllPropertyFlags(CPF_BlueprintReadOnly);
+
+							FAdvancedWidgetsModule& Module = FAdvancedWidgetsModule::GetModule();
+							TSharedPtr<SWidget> ValueWidget = Module.GetPropertyValueFactory().Generate(Args);
+							if (!ValueWidget)
+							{
+								ValueWidget = Module.GetPropertyValueFactory().GenerateDefault(Args);
+							}
+
+							if (ValueWidget)
+							{
+								return SNew(SHorizontalBox)
+								+ SHorizontalBox::Slot()
+								.AutoWidth()
+								[
+									ValueWidget.ToSharedRef()
+								]
+								+ SHorizontalBox::Slot()
+								.FillWidth(1.f)
+								[
+									SNullWidget::NullWidget
+								];
+							}
+						}
+					}
+				}
+			}
+
+			if (ColumnName == ColumnName_FieldPreWidget || ColumnName == ColumnName_FieldPostWidget)
+			{
+				TSharedPtr<FTreeNode> ItemPin = Item.Pin();
+				TSharedPtr<FPropertyViewerImpl> PropertyViewOwnerPin = PropertyViewOwner.Pin();
+				if (ItemPin && PropertyViewOwnerPin)
+				{
+					SPropertyViewer::FGetFieldWidget& OnGetWidget = (ColumnName == ColumnName_FieldPreWidget) ? PropertyViewOwnerPin->OnGetPreSlot : PropertyViewOwnerPin->OnGetPostSlot;
+					TSharedPtr<SWidget> PreWidget;
+					if (TSharedPtr<FContainer> ContainerPin = ItemPin->GetContainer())
+					{
+						PreWidget = OnGetWidget.Execute(ContainerPin->GetIdentifier(), FFieldVariant(ContainerPin->GetStruct()));
+					}
+					else if (FFieldVariant FieldVariant = ItemPin->GetField())
+					{
+						PreWidget = OnGetWidget.Execute(SPropertyViewer::FHandle(), FieldVariant);
+					}
+
+					if (PreWidget)
+					{
+						return PreWidget.ToSharedRef();
+					}
+				}
+			}
+
+			return SNullWidget::NullWidget;
+		}
+
+	private:
+		TWeakPtr<FTreeNode> Item;
+		TSharedPtr<SWidget> FieldWidget;
+		TWeakPtr<FPropertyViewerImpl> PropertyViewOwner;
+	};
+
+
+	TSharedRef<SWidget> FieldWidget = ItemWidget ? ItemWidget.ToSharedRef() : SNullWidget::NullWidget;
+	if (bUseRows)
+	{
+		return SNew(SMultiRowType, AsShared(), OwnerTable, Item.ToSharedRef(), FieldWidget)
+			.Padding(0.0f);
+	}
+
+	using SSimpleRowType = STableRow<TSharedPtr<FTreeNode>>;
+	return SNew(SSimpleRowType, OwnerTable)
+		.Padding(0.0f)
+		.Content()
+		[
+			FieldWidget
+		];
+}
+
+
+void FPropertyViewerImpl::HandleGetChildren(TSharedPtr<FTreeNode> InParent, TArray<TSharedPtr<FTreeNode>>& OutChildren)
+{
+	if (!InParent->bChildGenerated)
+	{			
+		// Do not build when filtering (only search in what it's already been built)
+		if (FilterHandler == nullptr || !FilterHandler->GetIsEnabled())
+		{
+			InParent->BuildChildNodes(*FieldIterator, *FieldExpander);
+		}
+	}
+	OutChildren = InParent->ChildNodes;
+}
+
+
+TSharedPtr<SWidget> FPropertyViewerImpl::HandleContextMenuOpening()
+{
+	if (OnContextMenuOpening.IsBound())
+	{
+		TArray<TSharedPtr<FTreeNode>> Items = TreeWidget->GetSelectedItems();
+		if (Items.Num() == 1 && Items[0].IsValid())
+		{
+			if (TSharedPtr<FContainer> ContainerPin = Items[0]->GetContainer())
+			{
+				return OnContextMenuOpening.Execute(ContainerPin->GetIdentifier(), FFieldVariant(ContainerPin->GetStruct()));
+			}
+			else if (FFieldVariant FieldVariant = Items[0]->GetField())
+			{
+				return OnContextMenuOpening.Execute(SPropertyViewer::FHandle(), FieldVariant);
+			}
+		}
+	}
+	return TSharedPtr<SWidget>();
+}
+
+
+void FPropertyViewerImpl::HandleSelectionChanged(TSharedPtr<FTreeNode> Item, ESelectInfo::Type SelectionType)
+{
+	if (OnSelectionChanged.IsBound())
+	{
+		if (TSharedPtr<FContainer> ContainerPin = Item->GetContainer())
+		{
+			OnSelectionChanged.Execute(ContainerPin->GetIdentifier(), FFieldVariant(ContainerPin->GetStruct()), SelectionType);
+		}
+		else if (FFieldVariant FieldVariant = Item->GetField())
+		{
+			OnSelectionChanged.Execute(SPropertyViewer::FHandle(), FieldVariant, SelectionType);
+		}
+	}
+}
+
+
+void FPropertyViewerImpl::HandleSearchChanged(const FText& InFilterText)
+{
+	if (SearchBoxWidget)
+	{
+		SearchBoxWidget->SetError(SetRawFilterText(InFilterText));
+	}
+}
+
+
+#if WITH_EDITOR
+void FPropertyViewerImpl::HandleBlueprintCompiled()
+{
+	bool bRemoved = false;
+	for (int32 Index = TreeSource.Num()-1; Index >= 0 ; --Index)
+	{
+		if (TSharedPtr<FContainer> Container = TreeSource[Index]->GetContainer())
+		{
+			if (!Container->IsValid())
+			{
+				TreeSource.RemoveAt(Index);
+				bRemoved = true;
+			}
+		}
+	}
+
+	for (int32 Index = Containers.Num()-1; Index >= 0; --Index)
+	{
+		if (!Containers[Index]->IsValid())
+		{
+			Containers.RemoveAt(Index);
+		}
+	}
+
+	if (bRemoved)
+	{
+		if (FilterHandler)
+		{
+			FilterHandler->RefreshAndFilterTree();
+		}
+		else
+		{
+			TreeWidget->RequestTreeRefresh();
+		}
+	}
+}
+#endif //WITH_EDITOR
+
+
+} //namespace
+
+#undef LOCTEXT_NAMESPACE
diff --git a/Engine/Source/Runtime/AdvancedWidgets/Private/Widgets/PropertyViewer/SPropertyViewerImpl.h b/Engine/Source/Runtime/AdvancedWidgets/Private/Widgets/PropertyViewer/SPropertyViewerImpl.h
new file mode 100644
--- /dev/null
+++ b/Engine/Source/Runtime/AdvancedWidgets/Private/Widgets/PropertyViewer/SPropertyViewerImpl.h
@@ -0,0 +1,234 @@
+// Copyright Epic Games, Inc. All Rights Reserved.
+
+#pragma once
+
+#include "Framework/PropertyViewer/PropertyPath.h"
+#include "Widgets/PropertyViewer/SPropertyViewer.h"
+#include "Types/SlateEnums.h"
+#include "UObject/WeakObjectPtrTemplates.h"
+#include "UObject/Field.h"
+
+class ITableRow;
+class SSearchBox;
+class STableViewBase;
+template <typename ItemType>
+class STreeView;
+template<typename ItemType>
+class TreeFilterHandler;
+template<typename ItemType>
+class TTextFilter;
+
+namespace UE::PropertyViewer
+{
+class IFieldIterator;
+class IFieldExpander;
+class SFieldName;
+}
+
+namespace UE::PropertyViewer::Private
+{
+
+/**
+ * FContainer
+ */
+struct FContainer
+{
+	FContainer(SPropertyViewer::FHandle InIdentifier, UObject* InstanceToDisplay);
+	FContainer(SPropertyViewer::FHandle InIdentifier, const UStruct* ClassToDisplay);
+	FContainer(SPropertyViewer::FHandle InIdentifier, const UScriptStruct* Struct, void* Data);
+
+public:
+	SPropertyViewer::FHandle GetIdentifier() const
+	{
+		return Identifier;
+	}
+
+	void* GetBuffer() const
+	{
+		UObject* PinObjectInstance = ObjectInstance.Get();
+		return bIsObject ? reinterpret_cast<void*>(&PinObjectInstance) : StructInstance;
+	}
+
+	const UStruct* GetStruct()
+	{
+		return Container.Get();
+	}
+
+	bool CanEdit() const
+	{
+		return GetBuffer() != nullptr;
+	}
+
+	bool IsInstance() const
+	{
+		return GetBuffer() != nullptr;
+	}
+	
+	bool IsObjectInstance() const
+	{
+		return bIsObject;
+	}
+
+	UObject* GetObjectInstance() const
+	{
+		return ObjectInstance.Get();
+	}
+
+	bool IsScriptStructInstance() const
+	{
+		return StructInstance != nullptr;
+	}
+
+	void* GetScriptStructInstance() const
+	{
+		return StructInstance;
+	}
+
+	bool IsValid() const;
+
+private:
+	SPropertyViewer::FHandle Identifier;
+	TWeakObjectPtr<const UStruct> Container;
+	TWeakObjectPtr<UObject> ObjectInstance;
+	void* StructInstance = nullptr;
+	bool bIsObject = false;
+};
+
+
+/**
+ * FTreeNode
+ */
+struct FTreeNode : public TSharedFromThis<FTreeNode>
+{
+public:
+	static TSharedRef<FTreeNode> MakeContainer(const TSharedPtr<FContainer>& InContainer, TOptional<FText> InDisplayName);
+	static TSharedRef<FTreeNode> MakeField(TSharedPtr<FTreeNode> InParent, const FProperty* Property, TOptional<FText> InDisplayName);
+	static TSharedRef<FTreeNode> MakeField(TSharedPtr<FTreeNode> InParent, const UFunction* Function, TOptional<FText> InDisplayName);
+
+private:
+	TWeakPtr<FTreeNode> ParentNode;
+
+	//Either a container, property or function
+	TWeakPtr<FContainer> Container;
+	const FProperty* Property = nullptr;
+	TWeakObjectPtr<const UFunction> Function;
+
+	TOptional<FText> OverrideDisplayName;
+
+public:
+	TWeakPtr<SFieldName> PropertyWidget;
+	TArray<TSharedPtr<FTreeNode>, TInlineAllocator<1>> ChildNodes;
+	bool bChildGenerated = false;
+
+public:
+	bool IsContainer() const
+	{
+		return Container.IsValid();
+	}
+
+	TSharedPtr<FContainer> GetContainer() const
+	{
+		return Container.Pin();
+	}
+
+	bool IsField() const
+	{
+		return Property != nullptr  || Function.Get() != nullptr;
+	}
+	
+	FFieldVariant GetField() const
+	{
+		return Property != nullptr ? FFieldVariant(Property) : Function.Get() ? FFieldVariant(Function.Get()) : FFieldVariant();
+	}
+
+	FPropertyPath GetPropertyPath() const;
+
+	TSharedPtr<FContainer> GetOwnerContainer() const;
+
+	TOptional<FText> GetOverrideDisplayName() const
+	{
+		return OverrideDisplayName;
+	}
+
+	void GetFilterStrings(TArray<FString>& OutStrings) const;
+
+	void BuildChildNodes(IFieldIterator& FieldIterator, IFieldExpander& FieldExpander);
+
+private:
+	void BuildChildNodesRecursive(IFieldIterator& FieldIterator, IFieldExpander& FieldExpander, int32 RecursiveCount);
+
+	static bool Sort(const TSharedPtr<FTreeNode>& NodeA, const TSharedPtr<FTreeNode>& NodeB);
+};
+
+
+/**
+ * FPropertyViewerImpl
+ */
+class FPropertyViewerImpl : public TSharedFromThis<FPropertyViewerImpl>
+{
+public:
+	FPropertyViewerImpl(const SPropertyViewer::FArguments& InArgs);
+	FPropertyViewerImpl(const FPropertyViewerImpl&) = delete;
+	FPropertyViewerImpl& operator=(const FPropertyViewerImpl&) = delete;
+	~FPropertyViewerImpl();
+
+private:
+	TArray<TSharedPtr<FContainer>> Containers;
+
+	TSharedPtr<SSearchBox> SearchBoxWidget;
+	TSharedPtr<STreeView<TSharedPtr<FTreeNode>>> TreeWidget;
+	TArray<TSharedPtr<FTreeNode>> TreeSource;
+	TArray<TSharedPtr<FTreeNode>> FilteredTreeSource;
+
+	using FTextFilter = TTextFilter<TSharedPtr<FTreeNode>>;
+	TSharedPtr<FTextFilter> SearchFilter;
+	using FTreeFilter = TreeFilterHandler<TSharedPtr<FTreeNode>>;
+	TSharedPtr<FTreeFilter> FilterHandler;
+
+	SPropertyViewer::FGetFieldWidget OnGetPreSlot;
+	SPropertyViewer::FGetFieldWidget OnGetPostSlot;
+	SPropertyViewer::FOnContextMenuOpening OnContextMenuOpening;
+	SPropertyViewer::FOnSelectionChanged OnSelectionChanged;
+	SPropertyViewer::EPropertyVisibility PropertyVisibility;
+	bool bSanitizeName = false;
+	bool bShowFieldIcon = false;
+	bool bUseRows = false;
+	
+	IFieldIterator* FieldIterator = nullptr;
+	IFieldExpander* FieldExpander = nullptr;
+	INotifyHook* NotifyHook = nullptr;
+	bool bOwnFieldIterator = false;
+	bool bOwnFieldExpander = false;
+
+public:
+	TSharedRef<SWidget> Construct(const SPropertyViewer::FArguments& InArgs);
+
+	void AddContainer(SPropertyViewer::FHandle Identifier, const UStruct* Struct);
+	void AddContainerInstance(SPropertyViewer::FHandle Identifier, UObject* Object);
+	void AddContainerInstance(SPropertyViewer::FHandle Identifier, const UScriptStruct* Struct, void* Data);
+	void Remove(SPropertyViewer::FHandle Identifier);
+	void RemoveAll();
+
+private:
+	void AddContainerInternal(SPropertyViewer::FHandle Identifier, TSharedPtr<FContainer>& NewContainer);
+
+	TSharedRef<SWidget> CreateSearch();
+	TSharedRef<SWidget> CreateTree(bool bHasPreWidget, bool bShowPropertyValue, bool bHasPostWidget);
+
+	void HandleSearchChanged(const FText& InFilterText);
+	FText SetRawFilterText(const FText& InFilterText);
+	void SetHighlightTextRecursive(const TSharedPtr<FTreeNode>& OwnerNode, const FText& HighlightText);
+	void HandleGetFilterStrings(TSharedPtr<FTreeNode> Item, TArray<FString>& OutStrings);
+
+	TSharedRef<ITableRow> HandleGenerateRow(TSharedPtr<FTreeNode> Item, const TSharedRef<STableViewBase>& OwnerTable);
+	void HandleGetChildren(TSharedPtr<FTreeNode> InParent, TArray<TSharedPtr<FTreeNode>>& OutChildren);
+
+	TSharedPtr<SWidget> HandleContextMenuOpening();
+	void HandleSelectionChanged(TSharedPtr<FTreeNode> Item, ESelectInfo::Type SelectionType);
+
+#if WITH_EDITOR
+	void HandleBlueprintCompiled();
+#endif //WITH_EDITOR
+};
+
+} //namespace
diff --git a/Engine/Source/Runtime/AdvancedWidgets/Private/Widgets/PropertyViewer/SStringPropertyValue.cpp b/Engine/Source/Runtime/AdvancedWidgets/Private/Widgets/PropertyViewer/SStringPropertyValue.cpp
new file mode 100644
--- /dev/null
+++ b/Engine/Source/Runtime/AdvancedWidgets/Private/Widgets/PropertyViewer/SStringPropertyValue.cpp
@@ -0,0 +1,106 @@
+// Copyright Epic Games, Inc. All Rights Reserved.
+
+#include "Widgets/PropertyViewer/SStringPropertyValue.h"
+
+#include "Framework/PropertyViewer/INotifyHook.h"
+#include "Styling/SlateTypes.h"
+#include "Styling/AppStyle.h"
+#include "UObject/TextProperty.h"
+#include "UObject/UnrealType.h"
+
+#include "Widgets/Input/SEditableText.h"
+
+namespace UE::PropertyViewer
+{
+
+TSharedPtr<SWidget> SStringPropertyValue::CreateInstance(const FPropertyValueFactory::FGenerateArgs Args)
+{
+	return SNew(SStringPropertyValue)
+		.Path(Args.Path)
+		.NotifyHook(Args.NotifyHook)
+		.IsEnabled(Args.bCanEditValue);
+}
+
+
+void SStringPropertyValue::Construct(const FArguments& InArgs)
+{
+	Path = InArgs._Path;
+	NotifyHook = InArgs._NotifyHook;
+
+	const FProperty* Property = Path.GetLastProperty();
+	if (CastField<const FStrProperty>(Property) || CastField<const FTextProperty>(Property))
+	{
+		if (Property->ArrayDim == 1)
+		{
+			ChildSlot
+				[
+					SNew(SEditableText)
+					.SelectAllTextWhenFocused(true)
+				.Text(this, &SStringPropertyValue::GetText)
+				.OnTextCommitted(this, &SStringPropertyValue::OnTextCommitted)
+				];
+		}
+	}
+}
+
+
+FText SStringPropertyValue::GetText() const
+{
+	if (const FStrProperty* Property = CastField<const FStrProperty>(Path.GetLastProperty()))
+	{
+		if (const void* Container = Path.GetContainerPtr())
+		{
+			FString OutString;
+			Property->GetValue_InContainer(Container, &OutString);
+			return FText::FromString(OutString);
+		}
+	}
+	if (const FTextProperty* Property = CastField<const FTextProperty>(Path.GetLastProperty()))
+	{
+		if (const void* Container = Path.GetContainerPtr())
+		{
+			FText OutText;
+			Property->GetValue_InContainer(Container, &OutText);
+			return OutText;
+		}
+	}
+	return FText::GetEmpty();
+}
+
+
+void SStringPropertyValue::OnTextCommitted(const FText& InText, ETextCommit::Type InCommitType)
+{
+	if (const FStrProperty* Property = CastField<const FStrProperty>(Path.GetLastProperty()))
+	{
+		if (void* Container = Path.GetContainerPtr())
+		{
+			FString ToSet = InText.ToString();
+			if (NotifyHook)
+			{
+				NotifyHook->OnPreValueChange(Path);
+			}
+			Property->SetValue_InContainer(Container, ToSet);
+			if (NotifyHook)
+			{
+				NotifyHook->OnPostValueChange(Path);
+			}
+		}
+	}
+	if (const FTextProperty* Property = CastField<const FTextProperty>(Path.GetLastProperty()))
+	{
+		if (void* Container = Path.GetContainerPtr())
+		{
+			if (NotifyHook)
+			{
+				NotifyHook->OnPreValueChange(Path);
+			}
+			Property->SetValue_InContainer(Container, InText);
+			if (NotifyHook)
+			{
+				NotifyHook->OnPostValueChange(Path);
+			}
+		}
+	}
+}
+
+} //namespace
diff --git a/Engine/Source/Runtime/AdvancedWidgets/Public/AdvancedWidgetsModule.h b/Engine/Source/Runtime/AdvancedWidgets/Public/AdvancedWidgetsModule.h
--- a/Engine/Source/Runtime/AdvancedWidgets/Public/AdvancedWidgetsModule.h
+++ b/Engine/Source/Runtime/AdvancedWidgets/Public/AdvancedWidgetsModule.h
@@ -6,12 +6,20 @@
 #include "Modules/ModuleInterface.h"
 
 
+namespace UE::PropertyViewer
+{
+class FPropertyValueFactory;
+}
+
+
 /**
  * Advanced Widgets module
  */
 class ADVANCEDWIDGETS_API FAdvancedWidgetsModule : public IModuleInterface
 {
 public:
+	static FAdvancedWidgetsModule& GetModule();
+	static FAdvancedWidgetsModule* GetModulePtr();
 
 	/**
 	 * Called right after the plugin DLL has been loaded and the plugin object has been created
@@ -23,6 +31,9 @@ public:
 	 */
 	virtual void ShutdownModule();
 
-	/** Advanced Widgets app identifier string */
-	static const FName AdvancedWidgetsAppIdentifier;
+public:
+	UE::PropertyViewer::FPropertyValueFactory& GetPropertyValueFactory() const;
+
+private:
+	TUniquePtr<UE::PropertyViewer::FPropertyValueFactory> PropertyValueFactory;
 };
diff --git a/Engine/Source/Runtime/AdvancedWidgets/Public/Framework/PropertyViewer/FieldIconFinder.h b/Engine/Source/Runtime/AdvancedWidgets/Public/Framework/PropertyViewer/FieldIconFinder.h
new file mode 100644
--- /dev/null
+++ b/Engine/Source/Runtime/AdvancedWidgets/Public/Framework/PropertyViewer/FieldIconFinder.h
@@ -0,0 +1,89 @@
+// Copyright Epic Games, Inc. All Rights Reserved.
+
+#pragma once
+
+#include "CoreMinimal.h"
+#include "Math/ColorList.h"
+#include "Styling/SlateColor.h"
+
+struct FSlateBrush;
+class FProperty;
+class UFunction;
+
+namespace UE::PropertyViewer
+{
+
+
+/** */
+struct FFieldColorSettings
+{
+	/** The default color is used only for types not specifically defined below. */
+	FLinearColor DefaultTypeColor = FLinearColor(0.750000f, 0.6f, 0.4f, 1.0f);
+
+	/** Boolean type color */
+	FLinearColor BooleanTypeColor = FColorList::Maroon;
+
+	/** Class type color */
+	FLinearColor ClassTypeColor = FLinearColor(0.1f, 0.0f, 0.5f, 1.0f);
+
+	/** Enum pin type color */
+	FLinearColor EnumTypeColor = FLinearColor(0.0f, 0.160000f, 0.131270f, 1.0f);
+
+	/** Integer pin type color */
+	FLinearColor IntTypeColor = FLinearColor(0.013575f, 0.770000f, 0.429609f, 1.0f);
+
+	/** Floating-point type color */
+	FLinearColor FloatTypeColor = FLinearColor(0.357667f, 1.0f, 0.060000f, 1.0f);
+
+	/** Name type color */
+	FLinearColor NameTypeColor = FLinearColor(0.607717f, 0.224984f, 1.0f, 1.0f);
+
+	/** Delegate type color */
+	FLinearColor DelegateTypeColor = FLinearColor(1.0f, 0.04f, 0.04f, 1.0f);
+
+	/** Object type color */
+	FLinearColor ObjectTypeColor = FLinearColor(0.0f, 0.4f, 0.910000f, 1.0f);
+
+	/** Interface type color */
+	FLinearColor InterfaceTypeColor = FLinearColor(0.8784f, 1.0f, 0.4f, 1.0f);
+
+	/** String type color */
+	FLinearColor StringTypeColor = FLinearColor(1.0f, 0.0f, 0.660537f, 1.0f);
+
+	/** Text type color */
+	FLinearColor TextTypeColor = FLinearColor(0.8f, 0.2f, 0.4f, 1.0f);
+
+	/** Default Struct pin type color */
+	FLinearColor DefaultStructTypeColor = FColorList::DarkSlateBlue;
+
+	/** Struct pin type color */
+	TMap<FString, FLinearColor> StructColors;
+};
+
+
+/** */
+class ADVANCEDWIDGETS_API FFieldIconFinder
+{
+public:
+	/** */
+	struct FFieldIcon
+	{
+		const FSlateBrush* Icon = nullptr;
+		FSlateColor Color;
+	};
+
+	using FFieldIconArray = TArray<FFieldIcon, TInlineAllocator<2>>;
+
+	/** */
+	static const FSlateBrush* GetIcon(const UObject* Object);
+
+	/** */
+	static FFieldIconArray GetFunctionIcon(const UFunction* Function);
+	static FFieldIconArray GetFunctionIcon(const UFunction* Function, const FFieldColorSettings& Settings);
+
+	/** */
+	static FFieldIconArray GetPropertyIcon(const FProperty* Property);
+	static FFieldIconArray GetPropertyIcon(const FProperty* Property, const FFieldColorSettings& Settings);
+};
+
+} //namespace
diff --git a/Engine/Source/Runtime/AdvancedWidgets/Public/Framework/PropertyViewer/IFieldExpander.h b/Engine/Source/Runtime/AdvancedWidgets/Public/Framework/PropertyViewer/IFieldExpander.h
new file mode 100644
--- /dev/null
+++ b/Engine/Source/Runtime/AdvancedWidgets/Public/Framework/PropertyViewer/IFieldExpander.h
@@ -0,0 +1,29 @@
+// Copyright Epic Games, Inc. All Rights Reserved.
+
+#pragma once
+
+#include "CoreMinimal.h"
+
+class FProperty;
+class UFunction;
+class UStruct;
+
+namespace UE::PropertyViewer
+{
+
+/** */
+class ADVANCEDWIDGETS_API IFieldExpander
+{
+public:
+	virtual bool CanExpandObject(const UClass*) const = 0;
+	virtual bool CanExpandFunction(const UFunction*) const = 0;
+	virtual ~IFieldExpander() = default;
+};
+
+class ADVANCEDWIDGETS_API FFieldExpander_NoExpand : public IFieldExpander
+{
+	virtual bool CanExpandObject(const UClass*) const override { return true; }
+	virtual bool CanExpandFunction(const UFunction*) const override { return false; }
+};
+
+} //namespace
diff --git a/Engine/Source/Runtime/AdvancedWidgets/Public/Framework/PropertyViewer/IFieldIterator.h b/Engine/Source/Runtime/AdvancedWidgets/Public/Framework/PropertyViewer/IFieldIterator.h
new file mode 100644
--- /dev/null
+++ b/Engine/Source/Runtime/AdvancedWidgets/Public/Framework/PropertyViewer/IFieldIterator.h
@@ -0,0 +1,28 @@
+// Copyright Epic Games, Inc. All Rights Reserved.
+
+#pragma once
+
+#include "CoreMinimal.h"
+#include "UObject/Field.h"
+
+class FProperty;
+class UFunction;
+class UStruct;
+
+namespace UE::PropertyViewer
+{
+
+/** */
+class ADVANCEDWIDGETS_API IFieldIterator
+{
+public:
+	virtual TArray<FFieldVariant> GetFields(const UStruct*) const = 0;
+	virtual ~IFieldIterator() = default;
+};
+
+class ADVANCEDWIDGETS_API FFieldIterator_BlueprintVisible : public IFieldIterator
+{
+	virtual TArray<FFieldVariant> GetFields(const UStruct*) const override;
+};
+
+} //namespace
diff --git a/Engine/Source/Runtime/AdvancedWidgets/Public/Framework/PropertyViewer/INotifyHook.h b/Engine/Source/Runtime/AdvancedWidgets/Public/Framework/PropertyViewer/INotifyHook.h
new file mode 100644
--- /dev/null
+++ b/Engine/Source/Runtime/AdvancedWidgets/Public/Framework/PropertyViewer/INotifyHook.h
@@ -0,0 +1,19 @@
+// Copyright Epic Games, Inc. All Rights Reserved.
+
+#pragma once 
+
+#include "CoreMinimal.h"
+#include "Framework/PropertyViewer/PropertyPath.h"
+
+namespace UE::PropertyViewer
+{
+
+/** */
+class INotifyHook
+{
+public:
+	virtual void OnPreValueChange(const FPropertyPath& Path) = 0;
+	virtual void OnPostValueChange(const FPropertyPath& Path) = 0;
+};
+
+} //namespace
diff --git a/Engine/Source/Runtime/AdvancedWidgets/Public/Framework/PropertyViewer/PropertyPath.h b/Engine/Source/Runtime/AdvancedWidgets/Public/Framework/PropertyViewer/PropertyPath.h
new file mode 100644
--- /dev/null
+++ b/Engine/Source/Runtime/AdvancedWidgets/Public/Framework/PropertyViewer/PropertyPath.h
@@ -0,0 +1,54 @@
+// Copyright Epic Games, Inc. All Rights Reserved.
+
+#pragma once
+
+#include "CoreMinimal.h"
+
+#include "Misc/TVariant.h"
+#include "UObject/FieldPath.h"
+#include "UObject/StructOnScope.h"
+#include "UObject/Object.h"
+#include "UObject/UnrealType.h"
+
+namespace UE::PropertyViewer
+{
+
+struct ADVANCEDWIDGETS_API FPropertyPath
+{
+public:
+	using FPropertyArray = TArray<const FProperty*, TInlineAllocator<1>>;
+
+	FPropertyPath() = default;
+	FPropertyPath(UObject* Object, const FProperty* Property);
+	FPropertyPath(UObject* Object, FPropertyArray Properties);
+	FPropertyPath(const UScriptStruct* ScriptStruct, void* Data, const FProperty* Property);
+	FPropertyPath(const UScriptStruct* ScriptStruct, void* Data, FPropertyArray Properties);
+
+public:
+	bool HasProperty() const
+	{
+		return Properties.Num() != 0;
+	}
+
+	const FProperty* GetLastProperty() const
+	{
+		return HasProperty() ? Properties.Last() : nullptr;
+	}
+
+	TArrayView<const FProperty* const> GetProperties() const
+	{
+		return MakeArrayView(Properties);
+	}
+
+	void* GetContainerPtr();
+	const void* GetContainerPtr() const;
+
+private:
+	TWeakObjectPtr<UObject> TopLevelContainer_Object = nullptr;
+	TWeakObjectPtr<const UScriptStruct> TopLevelContainer_ScriptStruct = nullptr;
+	void* TopLevelContainer_ScriptStructData = nullptr;
+
+	FPropertyArray Properties;
+};
+
+} //namespace
diff --git a/Engine/Source/Runtime/AdvancedWidgets/Public/Framework/PropertyViewer/PropertyValueFactory.h b/Engine/Source/Runtime/AdvancedWidgets/Public/Framework/PropertyViewer/PropertyValueFactory.h
new file mode 100644
--- /dev/null
+++ b/Engine/Source/Runtime/AdvancedWidgets/Public/Framework/PropertyViewer/PropertyValueFactory.h
@@ -0,0 +1,80 @@
+// Copyright Epic Games, Inc. All Rights Reserved.
+
+#pragma once 
+
+#include "CoreMinimal.h"
+#include "Modules/ModuleInterface.h"
+#include "Modules/ModuleManager.h"
+
+#include "Framework/PropertyViewer/PropertyPath.h"
+#include "Framework/SlateDelegates.h"
+#include "UObject/Field.h"
+#include "UObject/Class.h"
+
+class SWidget;
+
+namespace UE::PropertyViewer
+{
+class INotifyHook;
+
+namespace Private
+{
+class FPropertyValueFactoryImpl;
+} // namespace
+
+
+/** */
+class ADVANCEDWIDGETS_API FPropertyValueFactory
+{
+public:
+	struct FHandle
+	{
+	private:
+		friend FPropertyValueFactory;
+		int32 Id = 0;
+	public:
+		bool operator== (FHandle Other) const
+		{
+			return Id == Other.Id;
+		}
+		bool operator!= (FHandle Other) const
+		{
+			return Id != Other.Id;
+		}
+		bool IsValid() const
+		{
+			return Id != 0;
+		}
+	};
+
+	struct FGenerateArgs
+	{
+		FPropertyPath Path;
+		INotifyHook* NotifyHook = nullptr;
+		bool bCanEditValue = false;
+	};
+
+	/** */
+	DECLARE_DELEGATE_RetVal_OneParam(TSharedPtr<SWidget>, FOnGenerate, FGenerateArgs);
+
+public:
+	FPropertyValueFactory();
+	FPropertyValueFactory(const FPropertyValueFactory&) = delete;
+	FPropertyValueFactory& operator=(const FPropertyValueFactory&) = delete;
+
+	TSharedPtr<SWidget> Generate(FGenerateArgs Args) const;
+	TSharedPtr<SWidget> GenerateDefault(FGenerateArgs Args) const;
+
+	bool HasCustomPropertyValue(const FFieldClass* Property) const;
+	bool HasCustomPropertyValue(const UStruct* Struct) const;
+
+	FHandle Register(const FFieldClass* Field, FOnGenerate OnGenerate);
+	FHandle Register(const UStruct* Struct, FOnGenerate OnGenerate);
+	void Unregister(FHandle Handle);
+
+private:
+	static FHandle MakeHandle();
+	TPimplPtr<Private::FPropertyValueFactoryImpl> Impl;
+};
+
+} //namespace
diff --git a/Engine/Source/Runtime/AdvancedWidgets/Public/Styling/AdvancedWidgetsStyle.h b/Engine/Source/Runtime/AdvancedWidgets/Public/Styling/AdvancedWidgetsStyle.h
new file mode 100644
--- /dev/null
+++ b/Engine/Source/Runtime/AdvancedWidgets/Public/Styling/AdvancedWidgetsStyle.h
@@ -0,0 +1,36 @@
+// Copyright Epic Games, Inc. All Rights Reserved.
+
+#pragma once
+
+#include "CoreMinimal.h"
+#include "Framework/PropertyViewer/FieldIconFinder.h"
+
+class ISlateStyle;
+class FSlateStyleSet;
+
+namespace UE::AdvancedWidgets
+{
+/**
+ * Style for advanced widgets
+ */
+class ADVANCEDWIDGETS_API FAdvancedWidgetsStyle
+{
+public:
+	static void Create();
+	static void Destroy();
+	static const ISlateStyle& Get();
+
+	static const ::UE::PropertyViewer::FFieldColorSettings& GetColorSettings()
+	{
+		return ColorSettings;
+	}
+
+private:
+	/** Singleton instances of this style. */
+	static TUniquePtr<FSlateStyleSet> Instance;
+
+	/**  */
+	static ::UE::PropertyViewer::FFieldColorSettings ColorSettings;
+};
+
+} // namespace
\ No newline at end of file
diff --git a/Engine/Source/Runtime/AdvancedWidgets/Public/Widgets/PropertyViewer/SBoolPropertyValue.h b/Engine/Source/Runtime/AdvancedWidgets/Public/Widgets/PropertyViewer/SBoolPropertyValue.h
new file mode 100644
--- /dev/null
+++ b/Engine/Source/Runtime/AdvancedWidgets/Public/Widgets/PropertyViewer/SBoolPropertyValue.h
@@ -0,0 +1,38 @@
+// Copyright Epic Games, Inc. All Rights Reserved.
+
+#pragma once
+
+#include "Widgets/SCompoundWidget.h"
+#include "Framework/PropertyViewer/PropertyPath.h"
+#include "Framework/PropertyViewer/PropertyValueFactory.h"
+
+enum class ECheckBoxState : uint8;
+
+namespace UE::PropertyViewer
+{
+
+class INotifyHook;
+
+/** */
+class ADVANCEDWIDGETS_API SBoolPropertyValue : public SCompoundWidget
+{
+public:
+	static TSharedPtr<SWidget> CreateInstance(const FPropertyValueFactory::FGenerateArgs Args);
+
+public:
+	SLATE_BEGIN_ARGS(SBoolPropertyValue) {}
+		SLATE_ARGUMENT(FPropertyPath, Path);
+		SLATE_ARGUMENT(INotifyHook*, NotifyHook);
+	SLATE_END_ARGS()
+
+	void Construct(const FArguments& InArgs);
+
+private:
+	ECheckBoxState HandleIsChecked() const;
+	void HandleCheckStateChanged(ECheckBoxState NewState);
+
+	FPropertyPath Path;
+	INotifyHook* NotifyHook = nullptr;
+};
+
+} //namespace
diff --git a/Engine/Source/Runtime/AdvancedWidgets/Public/Widgets/PropertyViewer/SDefaultPropertyValue.h b/Engine/Source/Runtime/AdvancedWidgets/Public/Widgets/PropertyViewer/SDefaultPropertyValue.h
new file mode 100644
--- /dev/null
+++ b/Engine/Source/Runtime/AdvancedWidgets/Public/Widgets/PropertyViewer/SDefaultPropertyValue.h
@@ -0,0 +1,32 @@
+// Copyright Epic Games, Inc. All Rights Reserved.
+
+#pragma once
+
+#include "Widgets/SCompoundWidget.h"
+#include "Framework/PropertyViewer/PropertyPath.h"
+#include "Framework/PropertyViewer/PropertyValueFactory.h"
+
+
+namespace UE::PropertyViewer
+{
+
+/** */
+class ADVANCEDWIDGETS_API SDefaultPropertyValue : public SCompoundWidget
+{
+public:
+	static TSharedPtr<SWidget> CreateInstance(const FPropertyValueFactory::FGenerateArgs Args);
+
+public:
+	SLATE_BEGIN_ARGS(SDefaultPropertyValue) {}
+		SLATE_ARGUMENT(FPropertyPath, Path);
+	SLATE_END_ARGS()
+
+	void Construct(const FArguments& InArgs);
+
+private:
+	FText GetText() const;
+
+	FPropertyPath Path;
+};
+
+} //namespace
diff --git a/Engine/Source/Runtime/AdvancedWidgets/Public/Widgets/PropertyViewer/SEnumPropertyValue.h b/Engine/Source/Runtime/AdvancedWidgets/Public/Widgets/PropertyViewer/SEnumPropertyValue.h
new file mode 100644
--- /dev/null
+++ b/Engine/Source/Runtime/AdvancedWidgets/Public/Widgets/PropertyViewer/SEnumPropertyValue.h
@@ -0,0 +1,41 @@
+// Copyright Epic Games, Inc. All Rights Reserved.
+
+#pragma once
+
+#include "Widgets/SCompoundWidget.h"
+#include "Framework/PropertyViewer/PropertyPath.h"
+#include "Framework/PropertyViewer/PropertyValueFactory.h"
+
+
+namespace UE::PropertyViewer
+{
+
+class INotifyHook;
+
+/** */
+class ADVANCEDWIDGETS_API SEnumPropertyValue : public SCompoundWidget
+{
+public:
+	static TSharedPtr<SWidget> CreateInstance(const FPropertyValueFactory::FGenerateArgs Args);
+
+public:
+	SLATE_BEGIN_ARGS(SEnumPropertyValue) {}
+		SLATE_ARGUMENT(FPropertyPath, Path);
+		SLATE_ARGUMENT(INotifyHook*, NotifyHook);
+	SLATE_END_ARGS()
+
+	void Construct(const FArguments& InArgs);
+
+private:
+	FText GetText() const;
+	TSharedRef<SWidget> OnGetMenuContent();
+	//bool IsEnumEntryChecked(int32 Index) const;
+	void SetEnumEntry(int32 Index);
+	int32 GetCurrentValue() const;
+
+	TWeakObjectPtr<const UEnum> EnumType;
+	FPropertyPath Path;
+	INotifyHook* NotifyHook = nullptr;
+};
+
+} //namespace
diff --git a/Engine/Source/Runtime/AdvancedWidgets/Public/Widgets/PropertyViewer/SFieldIcon.h b/Engine/Source/Runtime/AdvancedWidgets/Public/Widgets/PropertyViewer/SFieldIcon.h
new file mode 100644
--- /dev/null
+++ b/Engine/Source/Runtime/AdvancedWidgets/Public/Widgets/PropertyViewer/SFieldIcon.h
@@ -0,0 +1,26 @@
+// Copyright Epic Games, Inc. All Rights Reserved.
+
+#pragma once
+
+#include "Widgets/SCompoundWidget.h"
+#include "Framework/PropertyViewer/FieldIconFinder.h"
+
+
+namespace UE::PropertyViewer
+{
+
+/** */
+class ADVANCEDWIDGETS_API SFieldIcon : public SCompoundWidget
+{
+public:
+	SLATE_BEGIN_ARGS(SFieldIcon) {}
+		SLATE_ARGUMENT(TOptional<const FFieldColorSettings>, OverrideColorSettings);
+	SLATE_END_ARGS()
+
+	void Construct(const FArguments& InArgs, const UClass* Class);
+	void Construct(const FArguments& InArgs, const UScriptStruct* Struct);
+	void Construct(const FArguments& InArgs, const FProperty* Property);
+	void Construct(const FArguments& InArgs, const UFunction* FunctionToDisplay);
+};
+
+} //namespace
diff --git a/Engine/Source/Runtime/AdvancedWidgets/Public/Widgets/PropertyViewer/SFieldName.h b/Engine/Source/Runtime/AdvancedWidgets/Public/Widgets/PropertyViewer/SFieldName.h
new file mode 100644
--- /dev/null
+++ b/Engine/Source/Runtime/AdvancedWidgets/Public/Widgets/PropertyViewer/SFieldName.h
@@ -0,0 +1,45 @@
+// Copyright Epic Games, Inc. All Rights Reserved.
+
+#pragma once
+
+#include "Widgets/SCompoundWidget.h"
+#include "UObject/Field.h"
+
+class IToolTip;
+class STextBlock;
+
+namespace UE::PropertyViewer
+{
+
+/** */
+class ADVANCEDWIDGETS_API SFieldName : public SCompoundWidget
+{
+public:
+	SLATE_BEGIN_ARGS(SFieldName)
+	{}
+		/** Show the class icon or the color that correspond to the property type. */
+		SLATE_ARGUMENT_DEFAULT(bool, bShowIcon) = true;
+		/** Sanitize the property name. */
+		SLATE_ARGUMENT_DEFAULT(bool, bSanitizeName) = false;
+		/** Override the DisplayName of the container name. */
+		SLATE_ARGUMENT(TOptional<FText>, OverrideDisplayName);
+		/** The current highlighted text. */
+		SLATE_ATTRIBUTE(FText, HighlightText)
+	SLATE_END_ARGS()
+
+	void Construct(const FArguments& InArgs, const UClass* Class);
+	void Construct(const FArguments& InArgs, const UScriptStruct* Struct);
+	void Construct(const FArguments& InArgs, const FProperty* Property);
+	void Construct(const FArguments& InArgs, const UFunction* FunctionToDisplay);
+
+	void SetHighlightText(TAttribute<FText> InHighlightText);
+
+private:
+	TSharedPtr<IToolTip> CreateToolTip() const;
+	void Construct(const FArguments& InArgs, const FText& DisplayName, TSharedPtr<SWidget> Icon);
+
+	TSharedPtr<STextBlock> NameBlock;
+	FFieldVariant Field;
+};
+
+} //namespace
diff --git a/Engine/Source/Runtime/AdvancedWidgets/Public/Widgets/PropertyViewer/SNumericPropertyValue.h b/Engine/Source/Runtime/AdvancedWidgets/Public/Widgets/PropertyViewer/SNumericPropertyValue.h
new file mode 100644
--- /dev/null
+++ b/Engine/Source/Runtime/AdvancedWidgets/Public/Widgets/PropertyViewer/SNumericPropertyValue.h
@@ -0,0 +1,48 @@
+// Copyright Epic Games, Inc. All Rights Reserved.
+
+#pragma once
+
+#include "Widgets/SCompoundWidget.h"
+#include "Framework/PropertyViewer/PropertyPath.h"
+#include "Framework/PropertyViewer/PropertyValueFactory.h"
+
+
+namespace UE::PropertyViewer
+{
+
+class INotifyHook;
+
+/** */
+class ADVANCEDWIDGETS_API SNumericPropertyValue : public SCompoundWidget
+{
+public:
+	static TSharedPtr<SWidget> CreateInstance(const FPropertyValueFactory::FGenerateArgs Args);
+
+public:
+	SLATE_BEGIN_ARGS(SNumericPropertyValue) {}
+		SLATE_ARGUMENT(FPropertyPath, Path);
+		SLATE_ARGUMENT(INotifyHook*, NotifyHook);
+	SLATE_END_ARGS()
+
+	void Construct(const FArguments& InArgs);
+
+private:
+	void OnEndSliderMovement_uint64(uint64 Value);
+	void OnEndSliderMovement_int64(int64 Value);
+	void OnEndSliderMovement_double(double Value);
+	uint64 GetCurrentValue_uint64() const;
+	int64 GetCurrentValue_int64() const;
+	double GetCurrentValue_double() const;
+	void OnValueChanged_uint64(uint64 Value);
+	void OnValueChanged_int64(int64 Value);
+	void OnValueChanged_double(double Value);
+	void OnValueCommitted_uint64(uint64 Value, ETextCommit::Type CommitInfo);
+	void OnValueCommitted_int64(int64 Value, ETextCommit::Type CommitInfo);
+	void OnValueCommitted_double(double Value, ETextCommit::Type CommitInfo);
+
+private:
+	FPropertyPath Path;
+	INotifyHook* NotifyHook = nullptr;
+};
+
+} //namespace
diff --git a/Engine/Source/Runtime/AdvancedWidgets/Public/Widgets/PropertyViewer/SPropertyViewer.h b/Engine/Source/Runtime/AdvancedWidgets/Public/Widgets/PropertyViewer/SPropertyViewer.h
new file mode 100644
--- /dev/null
+++ b/Engine/Source/Runtime/AdvancedWidgets/Public/Widgets/PropertyViewer/SPropertyViewer.h
@@ -0,0 +1,111 @@
+// Copyright Epic Games, Inc. All Rights Reserved.
+
+#pragma once
+
+#include "Widgets/SCompoundWidget.h"
+#include "UObject/Class.h"
+#include "UObject/Field.h"
+#include "UObject/StructOnScope.h"
+
+namespace UE::PropertyViewer
+{
+
+class IFieldIterator;
+class IFieldExpander;
+class INotifyHook;
+
+namespace Private
+{
+class FPropertyViewerImpl;
+}
+
+/** */
+class ADVANCEDWIDGETS_API SPropertyViewer : public SCompoundWidget
+{
+public:
+	class FHandle
+	{
+	private:
+		friend SPropertyViewer;
+		int32 Id = 0;
+	public:
+		bool operator==(FHandle Other) const
+		{
+			return Id == Other.Id;
+		}
+		bool operator!=(FHandle Other) const
+		{
+			return Id != Other.Id;
+		}
+		bool IsValid() const
+		{
+			return Id != 0;
+		}
+	};
+
+	enum class EPropertyVisibility
+	{
+		Hidden,
+		Visible,
+		Editable,
+	};
+
+public:
+	DECLARE_DELEGATE_RetVal_TwoParams(TSharedPtr<SWidget>, FGetFieldWidget, FHandle, const FFieldVariant);
+	DECLARE_DELEGATE_RetVal_TwoParams(TSharedPtr<SWidget>, FOnContextMenuOpening, FHandle, const FFieldVariant);
+	DECLARE_DELEGATE_ThreeParams(FOnSelectionChanged, FHandle, const FFieldVariant, ESelectInfo::Type);
+
+	SLATE_BEGIN_ARGS(SPropertyViewer)
+	{}
+		/** Allow to edit the instance property. */
+		SLATE_ARGUMENT_DEFAULT(EPropertyVisibility, PropertyVisibility) = EPropertyVisibility::Hidden;
+		/** Sanitize the field and container name. */
+		SLATE_ARGUMENT_DEFAULT(bool, bSanitizeName) = false;
+		/** Show the icon next to the field name. */
+		SLATE_ARGUMENT_DEFAULT(bool, bShowFieldIcon) = true;
+		/** Show a search box. */
+		SLATE_ARGUMENT_DEFAULT(bool, bShowSearchBox) = false;
+		/** Which properties/functions to show. FFieldIterator_BlueprintVisible is the default. */
+		SLATE_ARGUMENT_DEFAULT(IFieldIterator*, FieldIterator) = nullptr;
+		/** Which properties/functions that allow expansion. FFieldIterator_NoExpand is the default. */
+		SLATE_ARGUMENT_DEFAULT(IFieldExpander*, FieldExpander) = nullptr;
+		/** Callback when a property is modified. */
+		SLATE_ARGUMENT_DEFAULT(INotifyHook*, NotifyHook) = nullptr;
+		/** Slot for additional widget to go before the search box. */
+		SLATE_NAMED_SLOT(FArguments, SearchBoxPreSlot)
+		/** Slot for additional widget to go after the search box. */
+		SLATE_NAMED_SLOT(FArguments, SearchBoxPostSlot)
+		/** Slot for additional widget to go before the field or container widget. */
+		SLATE_EVENT(FGetFieldWidget, OnGetPreSlot)
+		/** Slot for additional widget to go after the field or container widget. */
+		SLATE_EVENT(FGetFieldWidget, OnGetPostSlot)
+		/** Context menu widget for the selected item. */
+		SLATE_EVENT(FOnContextMenuOpening, OnContextMenuOpening)
+		/** Delegate to invoke when selection changes. */
+		SLATE_EVENT(FOnSelectionChanged, OnSelectionChanged)
+	SLATE_END_ARGS()
+
+	void Construct(const FArguments& InArgs);
+	void Construct(const FArguments& InArgs, const UScriptStruct* Struct);
+	void Construct(const FArguments& InArgs, const UScriptStruct* Struct, void* Data);
+	void Construct(const FArguments& InArgs, const UClass* Class);
+	void Construct(const FArguments& InArgs, UObject* ObjectInstance);
+	void Construct(const FArguments& InArgs, const UFunction* Function);
+
+public:
+	FHandle AddContainer(const UScriptStruct* Struct);
+	FHandle AddContainer(const UClass* Class);
+	FHandle AddContainer(const UFunction* Function);
+	FHandle AddInstance(const UScriptStruct* Struct, void* Data);
+	FHandle AddInstance(UObject* ObjectInstance);
+
+	void Remove(FHandle Identifier);
+	void RemoveAll();
+
+private:
+	void ConstructInternal(const FArguments& InArgs);
+	static FHandle MakeContainerIdentifier();
+	TSharedPtr<Private::FPropertyViewerImpl> Implementation;
+};
+
+} //namespace
diff --git a/Engine/Source/Runtime/AdvancedWidgets/Public/Widgets/PropertyViewer/SStringPropertyValue.h b/Engine/Source/Runtime/AdvancedWidgets/Public/Widgets/PropertyViewer/SStringPropertyValue.h
new file mode 100644
--- /dev/null
+++ b/Engine/Source/Runtime/AdvancedWidgets/Public/Widgets/PropertyViewer/SStringPropertyValue.h
@@ -0,0 +1,37 @@
+// Copyright Epic Games, Inc. All Rights Reserved.
+
+#pragma once
+
+#include "Widgets/SCompoundWidget.h"
+#include "Framework/PropertyViewer/PropertyPath.h"
+#include "Framework/PropertyViewer/PropertyValueFactory.h"
+
+
+namespace UE::PropertyViewer
+{
+
+class INotifyHook;
+
+/** */
+class ADVANCEDWIDGETS_API SStringPropertyValue : public SCompoundWidget
+{
+public:
+	static TSharedPtr<SWidget> CreateInstance(const FPropertyValueFactory::FGenerateArgs Args);
+
+public:
+	SLATE_BEGIN_ARGS(SStringPropertyValue) {}
+		SLATE_ARGUMENT(FPropertyPath, Path);
+		SLATE_ARGUMENT(INotifyHook*, NotifyHook);
+	SLATE_END_ARGS()
+
+	void Construct(const FArguments& InArgs);
+
+private:
+	FText GetText() const;
+	void OnTextCommitted(const FText& InText, ETextCommit::Type InCommitType);
+
+	FPropertyPath Path;
+	INotifyHook* NotifyHook = nullptr;
+};
+
+} //namespace
