Fix CIS error introduced with CL20829550

diff --git a/Engine/Source/Runtime/AdvancedWidgets/Private/Widgets/PropertyViewer/SEnumPropertyValue.cpp b/Engine/Source/Runtime/AdvancedWidgets/Private/Widgets/PropertyViewer/SEnumPropertyValue.cpp
--- a/Engine/Source/Runtime/AdvancedWidgets/Private/Widgets/PropertyViewer/SEnumPropertyValue.cpp
+++ b/Engine/Source/Runtime/AdvancedWidgets/Private/Widgets/PropertyViewer/SEnumPropertyValue.cpp
@@ -110,9 +110,15 @@ TSharedRef<SWidget> SEnumPropertyValue::OnGetMenuContent()
 			if (!EnumPtr->HasMetaData(TEXT("Hidden"), Index))
 #endif
 			{
+#if WITH_EDITOR
+				FText Tooltip = EnumPtr->GetToolTipTextByIndex(Index);
+#else
+				FText Tooltip = FText::GetEmpty();
+#endif
+
 				MenuBuilder.AddMenuEntry(
 					EnumPtr->GetDisplayNameTextByIndex(Index),
-					EnumPtr->GetToolTipTextByIndex(Index),
+					Tooltip,
 					FSlateIcon(),
 					FUIAction
 					(
diff --git a/Engine/Source/Runtime/AdvancedWidgets/Private/Widgets/PropertyViewer/SFieldName.cpp b/Engine/Source/Runtime/AdvancedWidgets/Private/Widgets/PropertyViewer/SFieldName.cpp
--- a/Engine/Source/Runtime/AdvancedWidgets/Private/Widgets/PropertyViewer/SFieldName.cpp
+++ b/Engine/Source/Runtime/AdvancedWidgets/Private/Widgets/PropertyViewer/SFieldName.cpp
@@ -30,7 +30,11 @@ void SFieldName::Construct(const FArguments& InArgs, const UClass* Class)
 	}
 	else
 	{
+#if WITH_EDITORONLY_DATA
 		DisplayName = InArgs._bSanitizeName ? Class->GetDisplayNameText() : FText::FromName(Class->GetFName());
+#else
+		DisplayName = FText::FromName(Class->GetFName());
+#endif
 	}
 
 	Field = Class;
@@ -55,7 +59,11 @@ void SFieldName::Construct(const FArguments& InArgs, const UScriptStruct* Struct
 	}
 	else
 	{
+#if WITH_EDITORONLY_DATA
 		DisplayName = InArgs._bSanitizeName ? Struct->GetDisplayNameText() : FText::FromName(Struct->GetFName());
+#else
+		DisplayName = FText::FromName(Class->GetFName());
+#endif
 	}
 
 	Field = Struct;
@@ -80,7 +88,11 @@ void SFieldName::Construct(const FArguments& InArgs, const FProperty* Property)
 	}
 	else
 	{
+#if WITH_EDITORONLY_DATA
 		DisplayName = InArgs._bSanitizeName ? Property->GetDisplayNameText() : FText::FromName(Property->GetFName());
+#else
+		DisplayName = FText::FromName(Class->GetFName());
+#endif
 	}
 
 	Field = Property;
@@ -105,7 +117,11 @@ void SFieldName::Construct(const FArguments& InArgs, const UFunction* Function)
 	}
 	else
 	{
+#if WITH_EDITORONLY_DATA
 		DisplayName = InArgs._bSanitizeName ? Function->GetDisplayNameText() : FText::FromName(Function->GetFName());
+#else
+		DisplayName = FText::FromName(Class->GetFName());
+#endif
 	}
 
 	Field = Function;
@@ -145,7 +161,10 @@ void SFieldName::Construct(const FArguments& InArgs, const FText& DisplayName, T
 			.HighlightText(InArgs._HighlightText)
 		];
 	}
+
+#if WITH_EDITORONLY_DATA
 	SetToolTip(TAttribute<TSharedPtr<IToolTip>>::CreateSP(this, &SFieldName::CreateToolTip));
+#endif
 }
 
 
@@ -159,6 +178,7 @@ void SFieldName::SetHighlightText(TAttribute<FText> InHighlightText)
 
 TSharedPtr<IToolTip> SFieldName::CreateToolTip() const
 {
+#if WITH_EDITORONLY_DATA
 	if (FProperty* PropertyPtr = Field.Get<FProperty>())
 	{
 		return FSlateApplicationBase::Get().MakeToolTip(PropertyPtr->GetToolTipText());
@@ -167,6 +187,7 @@ TSharedPtr<IToolTip> SFieldName::CreateToolTip() const
 	{
 		return FSlateApplicationBase::Get().MakeToolTip(FieldPtr->GetToolTipText());
 	}
+#endif
 	return TSharedPtr<IToolTip>();
 }
 
diff --git a/Engine/Source/Runtime/AdvancedWidgets/Private/Widgets/PropertyViewer/SPropertyViewerImpl.cpp b/Engine/Source/Runtime/AdvancedWidgets/Private/Widgets/PropertyViewer/SPropertyViewerImpl.cpp
--- a/Engine/Source/Runtime/AdvancedWidgets/Private/Widgets/PropertyViewer/SPropertyViewerImpl.cpp
+++ b/Engine/Source/Runtime/AdvancedWidgets/Private/Widgets/PropertyViewer/SPropertyViewerImpl.cpp
@@ -201,19 +201,25 @@ void FTreeNode::GetFilterStrings(TArray<FString>& OutStrings) const
 	if (Property)
 	{
 		OutStrings.Add(Property->GetName());
+#if WITH_EDITORONLY_DATA
 		OutStrings.Add(Property->GetDisplayNameText().ToString());
+#endif
 	}
 	if (const UFunction* FunctionPtr = Function.Get())
 	{
 		OutStrings.Add(FunctionPtr->GetName());
+#if WITH_EDITORONLY_DATA
 		OutStrings.Add(FunctionPtr->GetDisplayNameText().ToString());
+#endif
 	}
 	if (const TSharedPtr<FContainer> ContainerPtr = Container.Pin())
 	{
 		if (const UStruct* StructPtr = ContainerPtr->GetStruct())
 		{
 			OutStrings.Add(StructPtr->GetName());
+#if WITH_EDITORONLY_DATA
 			OutStrings.Add(StructPtr->GetDisplayNameText().ToString());
+#endif
 		}
 	}
 	if (GetOverrideDisplayName().IsSet())
@@ -804,10 +810,10 @@ TSharedRef<ITableRow> FPropertyViewerImpl::HandleGenerateRow(TSharedPtr<FTreeNod
 						FFieldVariant Field = ItemPin->GetField();
 						if (!Field.IsUObject())
 						{
-							INotifyHook* NotifyHook = nullptr;
+							INotifyHook* LocalNotifyHook = nullptr;
 							if (const TSharedPtr<FPropertyViewerImpl> PropertyViewer = PropertyViewOwner.Pin())
 							{
-								NotifyHook = PropertyViewer->NotifyHook;
+								LocalNotifyHook = PropertyViewer->NotifyHook;
 							}
 
 							bool bCanEditContainer = false;
@@ -818,7 +824,7 @@ TSharedRef<ITableRow> FPropertyViewerImpl::HandleGenerateRow(TSharedPtr<FTreeNod
 
 							FPropertyValueFactory::FGenerateArgs Args;
 							Args.Path = ItemPin->GetPropertyPath();
-							Args.NotifyHook = NotifyHook;
+							Args.NotifyHook = LocalNotifyHook;
 							Args.bCanEditValue = bCanEditContainer
 								&& PropertyViewOwnerPin->PropertyVisibility == SPropertyViewer::EPropertyVisibility::Editable
 								&& Args.Path.GetLastProperty() != nullptr
diff --git a/Engine/Source/Runtime/AdvancedWidgets/Private/Widgets/PropertyViewer/SPropertyViewerImpl.h b/Engine/Source/Runtime/AdvancedWidgets/Private/Widgets/PropertyViewer/SPropertyViewerImpl.h
--- a/Engine/Source/Runtime/AdvancedWidgets/Private/Widgets/PropertyViewer/SPropertyViewerImpl.h
+++ b/Engine/Source/Runtime/AdvancedWidgets/Private/Widgets/PropertyViewer/SPropertyViewerImpl.h
@@ -46,7 +46,7 @@ public:
 	void* GetBuffer() const
 	{
 		UObject* PinObjectInstance = ObjectInstance.Get();
-		return bIsObject ? reinterpret_cast<void*>(&PinObjectInstance) : StructInstance;
+		return bIsObject ? reinterpret_cast<void*>(PinObjectInstance) : StructInstance;
 	}
 
 	const UStruct* GetStruct()
