Fix CIS error introduced with CL20829550

diff --git a/Engine/Source/Runtime/AdvancedWidgets/Private/Widgets/PropertyViewer/SFieldName.cpp b/Engine/Source/Runtime/AdvancedWidgets/Private/Widgets/PropertyViewer/SFieldName.cpp
--- a/Engine/Source/Runtime/AdvancedWidgets/Private/Widgets/PropertyViewer/SFieldName.cpp
+++ b/Engine/Source/Runtime/AdvancedWidgets/Private/Widgets/PropertyViewer/SFieldName.cpp
@@ -62,7 +62,7 @@ void SFieldName::Construct(const FArguments& InArgs, const UScriptStruct* Struct
 #if WITH_EDITORONLY_DATA
 		DisplayName = InArgs._bSanitizeName ? Struct->GetDisplayNameText() : FText::FromName(Struct->GetFName());
 #else
-		DisplayName = FText::FromName(Class->GetFName());
+		DisplayName = FText::FromName(Struct->GetFName());
 #endif
 	}
 
@@ -91,7 +91,7 @@ void SFieldName::Construct(const FArguments& InArgs, const FProperty* Property)
 #if WITH_EDITORONLY_DATA
 		DisplayName = InArgs._bSanitizeName ? Property->GetDisplayNameText() : FText::FromName(Property->GetFName());
 #else
-		DisplayName = FText::FromName(Class->GetFName());
+		DisplayName = FText::FromName(Property->GetFName());
 #endif
 	}
 
@@ -120,7 +120,7 @@ void SFieldName::Construct(const FArguments& InArgs, const UFunction* Function)
 #if WITH_EDITORONLY_DATA
 		DisplayName = InArgs._bSanitizeName ? Function->GetDisplayNameText() : FText::FromName(Function->GetFName());
 #else
-		DisplayName = FText::FromName(Class->GetFName());
+		DisplayName = FText::FromName(Function->GetFName());
 #endif
 	}
 
