PropertyViewer: Add support for FName property type. Remove duplicated local variable.

diff --git a/Engine/Source/Runtime/AdvancedWidgets/Private/Framework/PropertyViewer/PropertyValueFactory.cpp b/Engine/Source/Runtime/AdvancedWidgets/Private/Framework/PropertyViewer/PropertyValueFactory.cpp
--- a/Engine/Source/Runtime/AdvancedWidgets/Private/Framework/PropertyViewer/PropertyValueFactory.cpp
+++ b/Engine/Source/Runtime/AdvancedWidgets/Private/Framework/PropertyViewer/PropertyValueFactory.cpp
@@ -58,6 +58,7 @@ FPropertyValueFactory::FPropertyValueFactory()
 	Register(FEnumProperty::StaticClass(), FPropertyValueFactory::FOnGenerate::CreateStatic(SEnumPropertyValue::CreateInstance));
 	Register(FStrProperty::StaticClass(), FPropertyValueFactory::FOnGenerate::CreateStatic(SStringPropertyValue::CreateInstance));
 	Register(FTextProperty::StaticClass(), FPropertyValueFactory::FOnGenerate::CreateStatic(SStringPropertyValue::CreateInstance));
+	Register(FNameProperty::StaticClass(), FPropertyValueFactory::FOnGenerate::CreateStatic(SStringPropertyValue::CreateInstance));
 }
 
 
diff --git a/Engine/Source/Runtime/AdvancedWidgets/Private/Widgets/PropertyViewer/SPropertyViewerImpl.cpp b/Engine/Source/Runtime/AdvancedWidgets/Private/Widgets/PropertyViewer/SPropertyViewerImpl.cpp
--- a/Engine/Source/Runtime/AdvancedWidgets/Private/Widgets/PropertyViewer/SPropertyViewerImpl.cpp
+++ b/Engine/Source/Runtime/AdvancedWidgets/Private/Widgets/PropertyViewer/SPropertyViewerImpl.cpp
@@ -229,13 +229,13 @@ void FTreeNode::GetFilterStrings(TArray<FString>& OutStrings) const
 }
 
 
-void FTreeNode::BuildChildNodes(IFieldIterator& FieldIterator, IFieldExpander& FieldExpander)
+void FTreeNode::BuildChildNodes(IFieldIterator& FieldIterator, IFieldExpander& FieldExpander, bool bSortChildNode)
 {
-	BuildChildNodesRecursive(FieldIterator, FieldExpander, 2);
+	BuildChildNodesRecursive(FieldIterator, FieldExpander, bSortChildNode, 2);
 }
 
 
-void FTreeNode::BuildChildNodesRecursive(IFieldIterator& FieldIterator, IFieldExpander& FieldExpander, int32 RecursiveCount)
+void FTreeNode::BuildChildNodesRecursive(IFieldIterator& FieldIterator, IFieldExpander& FieldExpander, bool bSortChildNode, int32 RecursiveCount)
 {
 	if (RecursiveCount <= 0)
 	{
@@ -295,16 +295,19 @@ void FTreeNode::BuildChildNodesRecursive(IFieldIterator& FieldIterator, IFieldEx
 			if (const FProperty* PropertyIt = FieldIt.Get<FProperty>())
 			{
 				TSharedPtr<FTreeNode> Node = MakeField(AsShared(), PropertyIt, TOptional<FText>());
-				Node->BuildChildNodesRecursive(FieldIterator, FieldExpander, RecursiveCount);
+				Node->BuildChildNodesRecursive(FieldIterator, FieldExpander, bSortChildNode, RecursiveCount);
 			}
 			if (const UFunction* FunctionIt = FieldIt.Get<UFunction>())
 			{
 				TSharedPtr<FTreeNode> Node = MakeField(AsShared(), FunctionIt, TOptional<FText>());
-				Node->BuildChildNodesRecursive(FieldIterator, FieldExpander, RecursiveCount);
+				Node->BuildChildNodesRecursive(FieldIterator, FieldExpander, bSortChildNode, RecursiveCount);
 			}
 		}
 
-		ChildNodes.Sort(Sort);
+		if (bSortChildNode)
+		{
+			ChildNodes.Sort(Sort);
+		}
 	}
 
 	bChildGenerated = true;
@@ -377,6 +380,7 @@ FPropertyViewerImpl::FPropertyViewerImpl(const SPropertyViewer::FArguments& InAr
 	PropertyVisibility = InArgs._PropertyVisibility;
 	bSanitizeName = InArgs._bSanitizeName;
 	bShowFieldIcon = InArgs._bShowFieldIcon;
+	bSortChildNode = InArgs._bSortChildNode;
 
 #if WITH_EDITOR
 	if (GEditor)
@@ -510,7 +514,7 @@ void FPropertyViewerImpl::AddContainerInstance(SPropertyViewer::FHandle Identifi
 void FPropertyViewerImpl::AddContainerInternal(SPropertyViewer::FHandle Identifier, TSharedPtr<FContainer>& NewContainer)
 {
 	TSharedPtr<FTreeNode> NewNode = FTreeNode::MakeContainer(NewContainer, TOptional<FText>());
-	NewNode->BuildChildNodes(*FieldIterator, *FieldExpander);
+	NewNode->BuildChildNodes(*FieldIterator, *FieldExpander, bSortChildNode);
 	TreeSource.Add(NewNode);
 
 	if (TreeWidget)
@@ -810,12 +814,6 @@ TSharedRef<ITableRow> FPropertyViewerImpl::HandleGenerateRow(TSharedPtr<FTreeNod
 						FFieldVariant Field = ItemPin->GetField();
 						if (!Field.IsUObject())
 						{
-							INotifyHook* LocalNotifyHook = nullptr;
-							if (const TSharedPtr<FPropertyViewerImpl> PropertyViewer = PropertyViewOwner.Pin())
-							{
-								LocalNotifyHook = PropertyViewer->NotifyHook;
-							}
-
 							bool bCanEditContainer = false;
 							if (const TSharedPtr<FContainer> OwnerContainer = ItemPin->GetOwnerContainer())
 							{
@@ -824,7 +822,7 @@ TSharedRef<ITableRow> FPropertyViewerImpl::HandleGenerateRow(TSharedPtr<FTreeNod
 
 							FPropertyValueFactory::FGenerateArgs Args;
 							Args.Path = ItemPin->GetPropertyPath();
-							Args.NotifyHook = LocalNotifyHook;
+							Args.NotifyHook = PropertyViewOwnerPin->NotifyHook;
 							Args.bCanEditValue = bCanEditContainer
 								&& PropertyViewOwnerPin->PropertyVisibility == SPropertyViewer::EPropertyVisibility::Editable
 								&& Args.Path.GetLastProperty() != nullptr
@@ -914,7 +912,7 @@ void FPropertyViewerImpl::HandleGetChildren(TSharedPtr<FTreeNode> InParent, TArr
 		// Do not build when filtering (only search in what it's already been built)
 		if (FilterHandler == nullptr || !FilterHandler->GetIsEnabled())
 		{
-			InParent->BuildChildNodes(*FieldIterator, *FieldExpander);
+			InParent->BuildChildNodes(*FieldIterator, *FieldExpander, bSortChildNode);
 		}
 	}
 	OutChildren = InParent->ChildNodes;
diff --git a/Engine/Source/Runtime/AdvancedWidgets/Private/Widgets/PropertyViewer/SPropertyViewerImpl.h b/Engine/Source/Runtime/AdvancedWidgets/Private/Widgets/PropertyViewer/SPropertyViewerImpl.h
--- a/Engine/Source/Runtime/AdvancedWidgets/Private/Widgets/PropertyViewer/SPropertyViewerImpl.h
+++ b/Engine/Source/Runtime/AdvancedWidgets/Private/Widgets/PropertyViewer/SPropertyViewerImpl.h
@@ -152,10 +152,10 @@ public:
 
 	void GetFilterStrings(TArray<FString>& OutStrings) const;
 
-	void BuildChildNodes(IFieldIterator& FieldIterator, IFieldExpander& FieldExpander);
+	void BuildChildNodes(IFieldIterator& FieldIterator, IFieldExpander& FieldExpander, bool bSortChildNode);
 
 private:
-	void BuildChildNodesRecursive(IFieldIterator& FieldIterator, IFieldExpander& FieldExpander, int32 RecursiveCount);
+	void BuildChildNodesRecursive(IFieldIterator& FieldIterator, IFieldExpander& FieldExpander, bool bSortChildNode, int32 RecursiveCount);
 
 	static bool Sort(const TSharedPtr<FTreeNode>& NodeA, const TSharedPtr<FTreeNode>& NodeB);
 };
@@ -193,6 +193,7 @@ private:
 	bool bSanitizeName = false;
 	bool bShowFieldIcon = false;
 	bool bUseRows = false;
+	bool bSortChildNode = false;
 	
 	IFieldIterator* FieldIterator = nullptr;
 	IFieldExpander* FieldExpander = nullptr;
diff --git a/Engine/Source/Runtime/AdvancedWidgets/Private/Widgets/PropertyViewer/SStringPropertyValue.cpp b/Engine/Source/Runtime/AdvancedWidgets/Private/Widgets/PropertyViewer/SStringPropertyValue.cpp
--- a/Engine/Source/Runtime/AdvancedWidgets/Private/Widgets/PropertyViewer/SStringPropertyValue.cpp
+++ b/Engine/Source/Runtime/AdvancedWidgets/Private/Widgets/PropertyViewer/SStringPropertyValue.cpp
@@ -28,17 +28,17 @@ void SStringPropertyValue::Construct(const FArguments& InArgs)
 	NotifyHook = InArgs._NotifyHook;
 
 	const FProperty* Property = Path.GetLastProperty();
-	if (CastField<const FStrProperty>(Property) || CastField<const FTextProperty>(Property))
+	if (CastField<const FStrProperty>(Property) || CastField<const FTextProperty>(Property) || CastField<const FNameProperty>(Property))
 	{
 		if (Property->ArrayDim == 1)
 		{
 			ChildSlot
-				[
-					SNew(SEditableText)
-					.SelectAllTextWhenFocused(true)
+			[
+				SNew(SEditableText)
+				.SelectAllTextWhenFocused(true)
 				.Text(this, &SStringPropertyValue::GetText)
 				.OnTextCommitted(this, &SStringPropertyValue::OnTextCommitted)
-				];
+			];
 		}
 	}
 }
@@ -46,31 +46,40 @@ void SStringPropertyValue::Construct(const FArguments& InArgs)
 
 FText SStringPropertyValue::GetText() const
 {
-	if (const FStrProperty* Property = CastField<const FStrProperty>(Path.GetLastProperty()))
+	if (const FStrProperty* StrProperty = CastField<const FStrProperty>(Path.GetLastProperty()))
 	{
 		if (const void* Container = Path.GetContainerPtr())
 		{
 			FString OutString;
-			Property->GetValue_InContainer(Container, &OutString);
+			StrProperty->GetValue_InContainer(Container, &OutString);
 			return FText::FromString(OutString);
 		}
 	}
-	if (const FTextProperty* Property = CastField<const FTextProperty>(Path.GetLastProperty()))
+	else if (const FTextProperty* TextProperty = CastField<const FTextProperty>(Path.GetLastProperty()))
 	{
 		if (const void* Container = Path.GetContainerPtr())
 		{
 			FText OutText;
-			Property->GetValue_InContainer(Container, &OutText);
+			TextProperty->GetValue_InContainer(Container, &OutText);
 			return OutText;
 		}
 	}
+	else if (const FNameProperty* NameProperty = CastField<const FNameProperty>(Path.GetLastProperty()))
+	{
+		if (const void* Container = Path.GetContainerPtr())
+		{
+			FName OutText;
+			NameProperty->GetValue_InContainer(Container, &OutText);
+			return FText::FromName(OutText);
+		}
+	}
 	return FText::GetEmpty();
 }
 
 
 void SStringPropertyValue::OnTextCommitted(const FText& InText, ETextCommit::Type InCommitType)
 {
-	if (const FStrProperty* Property = CastField<const FStrProperty>(Path.GetLastProperty()))
+	if (const FStrProperty* StrProperty = CastField<const FStrProperty>(Path.GetLastProperty()))
 	{
 		if (void* Container = Path.GetContainerPtr())
 		{
@@ -79,22 +88,38 @@ void SStringPropertyValue::OnTextCommitted(const FText& InText, ETextCommit::Typ
 			{
 				NotifyHook->OnPreValueChange(Path);
 			}
-			Property->SetValue_InContainer(Container, ToSet);
+			StrProperty->SetValue_InContainer(Container, ToSet);
+			if (NotifyHook)
+			{
+				NotifyHook->OnPostValueChange(Path);
+			}
+		}
+	}
+	else if (const FTextProperty* TextProperty = CastField<const FTextProperty>(Path.GetLastProperty()))
+	{
+		if (void* Container = Path.GetContainerPtr())
+		{
+			if (NotifyHook)
+			{
+				NotifyHook->OnPreValueChange(Path);
+			}
+			TextProperty->SetValue_InContainer(Container, InText);
 			if (NotifyHook)
 			{
 				NotifyHook->OnPostValueChange(Path);
 			}
 		}
 	}
-	if (const FTextProperty* Property = CastField<const FTextProperty>(Path.GetLastProperty()))
+	else if (const FNameProperty* NameProperty = CastField<const FNameProperty>(Path.GetLastProperty()))
 	{
 		if (void* Container = Path.GetContainerPtr())
 		{
+			FName ToSet = *(InText.ToString());
 			if (NotifyHook)
 			{
 				NotifyHook->OnPreValueChange(Path);
 			}
-			Property->SetValue_InContainer(Container, InText);
+			NameProperty->SetValue_InContainer(Container, ToSet);
 			if (NotifyHook)
 			{
 				NotifyHook->OnPostValueChange(Path);
diff --git a/Engine/Source/Runtime/AdvancedWidgets/Public/Widgets/PropertyViewer/SPropertyViewer.h b/Engine/Source/Runtime/AdvancedWidgets/Public/Widgets/PropertyViewer/SPropertyViewer.h
--- a/Engine/Source/Runtime/AdvancedWidgets/Public/Widgets/PropertyViewer/SPropertyViewer.h
+++ b/Engine/Source/Runtime/AdvancedWidgets/Public/Widgets/PropertyViewer/SPropertyViewer.h
@@ -63,6 +63,8 @@ public:
 		SLATE_ARGUMENT_DEFAULT(bool, bSanitizeName) = false;
 		/** Show the icon next to the field name. */
 		SLATE_ARGUMENT_DEFAULT(bool, bShowFieldIcon) = true;
+		/** Sort the field after building the list of fields. */
+		SLATE_ARGUMENT_DEFAULT(bool, bSortChildNode) = false;
 		/** Show a search box. */
 		SLATE_ARGUMENT_DEFAULT(bool, bShowSearchBox) = false;
 		/** Which properties/functions to show. FFieldIterator_BlueprintVisible is the default. */
