Fix static analysis introduced with CL20844899
#rnx

diff --git a/Engine/Source/Runtime/AdvancedWidgets/Private/Framework/PropertyViewer/FieldIconFinder.cpp b/Engine/Source/Runtime/AdvancedWidgets/Private/Framework/PropertyViewer/FieldIconFinder.cpp
--- a/Engine/Source/Runtime/AdvancedWidgets/Private/Framework/PropertyViewer/FieldIconFinder.cpp
+++ b/Engine/Source/Runtime/AdvancedWidgets/Private/Framework/PropertyViewer/FieldIconFinder.cpp
@@ -111,6 +111,8 @@ FFieldIconFinder::FFieldIconArray FFieldIconFinder::GetFunctionIcon(const UFunct
 
 FFieldIconFinder::FFieldIconArray FFieldIconFinder::GetFunctionIcon(const UFunction* Function, const FFieldColorSettings& Settings)
 {
+	check(Function);
+
 	FFieldIcon Icon;
 	if (UFunction* OverrideFunc = Function->GetSuperFunction())
 	{
@@ -119,7 +121,7 @@ FFieldIconFinder::FFieldIconArray FFieldIconFinder::GetFunctionIcon(const UFunct
 	}
 	else
 	{
-		const bool bIsPureFunction = Function && Function->HasAnyFunctionFlags(FUNC_BlueprintPure);
+		const bool bIsPureFunction = Function->HasAnyFunctionFlags(FUNC_BlueprintPure);
 		Icon.Icon = FAppStyle::GetBrush(bIsPureFunction ? TEXT("GraphEditor.PureFunction_16x") : TEXT("GraphEditor.Function_16x"));
 	}
 
@@ -147,6 +149,8 @@ FFieldIconFinder::FFieldIconArray FFieldIconFinder::GetPropertyIcon(const FPrope
 
 FFieldIconFinder::FFieldIconArray FFieldIconFinder::GetPropertyIcon(const FProperty* Property, const FFieldColorSettings& Settings)
 {
+	check(Property);
+
 	FFieldIconArray Result;
 
 	if (const FArrayProperty* ArrayProperty = CastField<FArrayProperty>(Property))
