PropertyViewer: Align value widget with the field name. Customize the container display name. Add search ability when the seach bar is hidden.

diff --git a/Engine/Source/Runtime/AdvancedWidgets/Private/Widgets/PropertyViewer/SPropertyViewer.cpp b/Engine/Source/Runtime/AdvancedWidgets/Private/Widgets/PropertyViewer/SPropertyViewer.cpp
--- a/Engine/Source/Runtime/AdvancedWidgets/Private/Widgets/PropertyViewer/SPropertyViewer.cpp
+++ b/Engine/Source/Runtime/AdvancedWidgets/Private/Widgets/PropertyViewer/SPropertyViewer.cpp
@@ -24,66 +24,66 @@ void SPropertyViewer::Construct(const FArguments& InArgs)
 void SPropertyViewer::Construct(const FArguments& InArgs, const UScriptStruct* Struct)
 {
 	Implementation = MakeShared<Private::FPropertyViewerImpl>(InArgs);
-	Implementation->AddContainer(MakeContainerIdentifier(), Struct);
+	Implementation->AddContainer(MakeContainerIdentifier(), TOptional<FText>(), Struct);
 	ConstructInternal(InArgs);
 }
 void SPropertyViewer::Construct(const FArguments& InArgs, const UScriptStruct* Struct, void* InData)
 {
 	Implementation = MakeShared<Private::FPropertyViewerImpl>(InArgs);
-	Implementation->AddContainerInstance(MakeContainerIdentifier(), Struct, InData);
+	Implementation->AddContainerInstance(MakeContainerIdentifier(), TOptional<FText>(), Struct, InData);
 	ConstructInternal(InArgs);
 }
 void SPropertyViewer::Construct(const FArguments& InArgs, const UClass* Class)
 {
 	Implementation = MakeShared<Private::FPropertyViewerImpl>(InArgs);
-	Implementation->AddContainer(MakeContainerIdentifier(), Class);
+	Implementation->AddContainer(MakeContainerIdentifier(), TOptional<FText>(), Class);
 	ConstructInternal(InArgs);
 }
 void SPropertyViewer::Construct(const FArguments& InArgs, UObject* ObjectInstance)
 {
 	Implementation = MakeShared<Private::FPropertyViewerImpl>(InArgs);
-	Implementation->AddContainerInstance(MakeContainerIdentifier(), ObjectInstance);
+	Implementation->AddContainerInstance(MakeContainerIdentifier(), TOptional<FText>(), ObjectInstance);
 	ConstructInternal(InArgs);
 }
 void SPropertyViewer::Construct(const FArguments& InArgs, const UFunction* Function)
 {
 	Implementation = MakeShared<Private::FPropertyViewerImpl>(InArgs);
-	Implementation->AddContainer(MakeContainerIdentifier(), Function);
+	Implementation->AddContainer(MakeContainerIdentifier(), TOptional<FText>(), Function);
 	ConstructInternal(InArgs);
 }
 
 
-SPropertyViewer::FHandle SPropertyViewer::AddContainer(const UScriptStruct* Struct)
+SPropertyViewer::FHandle SPropertyViewer::AddContainer(const UScriptStruct* Struct, TOptional<FText> DisplayName)
 {
 	SPropertyViewer::FHandle Result = MakeContainerIdentifier();
-	Implementation->AddContainer(Result, Struct);
+	Implementation->AddContainer(Result, DisplayName, Struct);
 	return Result;
 }
-SPropertyViewer::FHandle SPropertyViewer::AddContainer(const UClass* Class)
+SPropertyViewer::FHandle SPropertyViewer::AddContainer(const UClass* Class, TOptional<FText> DisplayName)
 {
 	SPropertyViewer::FHandle Result = MakeContainerIdentifier();
-	Implementation->AddContainer(Result, Class);
+	Implementation->AddContainer(Result, DisplayName, Class);
 	return Result;
 }
-SPropertyViewer::FHandle SPropertyViewer::AddContainer(const UFunction* Function)
+SPropertyViewer::FHandle SPropertyViewer::AddContainer(const UFunction* Function, TOptional<FText> DisplayName)
 {
 	SPropertyViewer::FHandle Result = MakeContainerIdentifier();
-	Implementation->AddContainer(Result, Function);
+	Implementation->AddContainer(Result, DisplayName, Function);
 	return Result;
 }
 
 
-SPropertyViewer::FHandle SPropertyViewer::AddInstance(const UScriptStruct* Struct, void* InData)
+SPropertyViewer::FHandle SPropertyViewer::AddInstance(const UScriptStruct* Struct, void* InData, TOptional<FText> DisplayName)
 {
 	check(InData);
 	SPropertyViewer::FHandle Result = MakeContainerIdentifier();
-	Implementation->AddContainerInstance(Result, Struct, InData);
+	Implementation->AddContainerInstance(Result, DisplayName, Struct, InData);
 	return Result;
 }
-SPropertyViewer::FHandle SPropertyViewer::AddInstance(UObject* ObjectInstance)
+SPropertyViewer::FHandle SPropertyViewer::AddInstance(UObject* ObjectInstance, TOptional<FText> DisplayName)
 {
 	SPropertyViewer::FHandle Result = MakeContainerIdentifier();
-	Implementation->AddContainerInstance(Result, ObjectInstance);
+	Implementation->AddContainerInstance(Result, DisplayName, ObjectInstance);
 	return Result;
 }
 
@@ -96,11 +96,16 @@ void SPropertyViewer::Remove(FHandle Identifier)
 
 void SPropertyViewer::RemoveAll()
 {
-
 	Implementation->RemoveAll();
 }
 
 
+void SPropertyViewer::SetRawFilterText(const FText& InFilterText)
+{
+	Implementation->SetRawFilterText(InFilterText);
+}
+
+
 SPropertyViewer::FHandle SPropertyViewer::MakeContainerIdentifier()
 {
 	static int32 IdentifierGenerator = 0;
diff --git a/Engine/Source/Runtime/AdvancedWidgets/Private/Widgets/PropertyViewer/SPropertyViewerImpl.cpp b/Engine/Source/Runtime/AdvancedWidgets/Private/Widgets/PropertyViewer/SPropertyViewerImpl.cpp
--- a/Engine/Source/Runtime/AdvancedWidgets/Private/Widgets/PropertyViewer/SPropertyViewerImpl.cpp
+++ b/Engine/Source/Runtime/AdvancedWidgets/Private/Widgets/PropertyViewer/SPropertyViewerImpl.cpp
@@ -48,26 +48,29 @@ static FName ColumnName_FieldPostWidget = "FieldPostWidget";
 /**
  * FContainer
  */
-FContainer::FContainer(SPropertyViewer::FHandle InIdentifier, UObject* InstanceToDisplay)
+FContainer::FContainer(SPropertyViewer::FHandle InIdentifier, TOptional<FText> InDisplayName, const UStruct* ClassToDisplay)
 	: Identifier(InIdentifier)
-	, Container(InstanceToDisplay->GetClass())
-	, ObjectInstance(InstanceToDisplay)
-	, bIsObject(true)
+	, Container(ClassToDisplay)
+	, DisplayName(InDisplayName)
 {
 }
 
 
-FContainer::FContainer(SPropertyViewer::FHandle InIdentifier, const UStruct* ClassToDisplay)
+FContainer::FContainer(SPropertyViewer::FHandle InIdentifier, TOptional<FText> InDisplayName, UObject* InstanceToDisplay)
 	: Identifier(InIdentifier)
-	, Container(ClassToDisplay)
+	, Container(InstanceToDisplay->GetClass())
+	, ObjectInstance(InstanceToDisplay)
+	, DisplayName(InDisplayName)
+	, bIsObject(true)
 {
 }
 
 
-FContainer::FContainer(SPropertyViewer::FHandle InIdentifier, const UScriptStruct* Struct, void* Data)
+FContainer::FContainer(SPropertyViewer::FHandle InIdentifier, TOptional<FText> InDisplayName, const UScriptStruct* Struct, void* Data)
 	: Identifier(InIdentifier)
 	, Container(Struct)
 	, StructInstance(Data)
+	, DisplayName(InDisplayName)
 {
 }
 
@@ -412,6 +415,13 @@ FPropertyViewerImpl::~FPropertyViewerImpl()
 
 TSharedRef<SWidget> FPropertyViewerImpl::Construct(const SPropertyViewer::FArguments& InArgs)
 {
+	SearchFilter = MakeShared<FTextFilter>(FTextFilter::FItemToStringArray::CreateSP(this, &FPropertyViewerImpl::HandleGetFilterStrings));
+
+	FilterHandler = MakeShared<FTreeFilter>();
+	FilterHandler->SetFilter(SearchFilter.Get());
+	FilterHandler->SetRootItems(&TreeSource, &FilteredTreeSource);
+	FilterHandler->SetGetChildrenDelegate(FTreeFilter::FOnGetChildren::CreateRaw(this, &FPropertyViewerImpl::HandleGetChildren));
+
 	TSharedPtr<SHorizontalBox> SearchBox;
 	if (InArgs._bShowSearchBox || InArgs._SearchBoxPreSlot.Widget != SNullWidget::NullWidget || InArgs._SearchBoxPostSlot.Widget != SNullWidget::NullWidget)
 	{
@@ -484,27 +494,27 @@ TSharedRef<SWidget> FPropertyViewerImpl::Construct(const SPropertyViewer::FArgum
 }
 
 
-void FPropertyViewerImpl::AddContainer(SPropertyViewer::FHandle Identifier, const UStruct* Struct)
+void FPropertyViewerImpl::AddContainer(SPropertyViewer::FHandle Identifier, TOptional<FText> DisplayName, const UStruct* Struct)
 {
-	TSharedPtr<FContainer> NewContainer = MakeShared<FContainer>(Identifier, Struct);
+	TSharedPtr<FContainer> NewContainer = MakeShared<FContainer>(Identifier, DisplayName, Struct);
 	Containers.Add(NewContainer);
 
 	AddContainerInternal(Identifier, NewContainer);
 }
 
 
-void FPropertyViewerImpl::AddContainerInstance(SPropertyViewer::FHandle Identifier, UObject* Object)
+void FPropertyViewerImpl::AddContainerInstance(SPropertyViewer::FHandle Identifier, TOptional<FText> DisplayName, UObject* Object)
 {
-	TSharedPtr<FContainer> NewContainer = MakeShared<FContainer>(Identifier, Object);
+	TSharedPtr<FContainer> NewContainer = MakeShared<FContainer>(Identifier, DisplayName, Object);
 	Containers.Add(NewContainer);
 
 	AddContainerInternal(Identifier, NewContainer);
 }
 
 
-void FPropertyViewerImpl::AddContainerInstance(SPropertyViewer::FHandle Identifier, const UScriptStruct* Struct, void* Data)
+void FPropertyViewerImpl::AddContainerInstance(SPropertyViewer::FHandle Identifier, TOptional<FText> DisplayName, const UScriptStruct* Struct, void* Data)
 {
-	TSharedPtr<FContainer> NewContainer = MakeShared<FContainer>(Identifier, Struct, Data);
+	TSharedPtr<FContainer> NewContainer = MakeShared<FContainer>(Identifier, DisplayName, Struct, Data);
 	Containers.Add(NewContainer);
 
 	AddContainerInternal(Identifier, NewContainer);
@@ -513,7 +523,7 @@ void FPropertyViewerImpl::AddContainerInstance(SPropertyViewer::FHandle Identifi
 
 void FPropertyViewerImpl::AddContainerInternal(SPropertyViewer::FHandle Identifier, TSharedPtr<FContainer>& NewContainer)
 {
-	TSharedPtr<FTreeNode> NewNode = FTreeNode::MakeContainer(NewContainer, TOptional<FText>());
+	TSharedPtr<FTreeNode> NewNode = FTreeNode::MakeContainer(NewContainer, NewContainer->GetDisplayName());
 	NewNode->BuildChildNodes(*FieldIterator, *FieldExpander, bSortChildNode);
 	TreeSource.Add(NewNode);
 
@@ -589,20 +599,19 @@ void FPropertyViewerImpl::RemoveAll()
 
 TSharedRef<SWidget> FPropertyViewerImpl::CreateSearch()
 {
-	SearchFilter = MakeShared<FTextFilter>(FTextFilter::FItemToStringArray::CreateSP(this, &FPropertyViewerImpl::HandleGetFilterStrings));
-
-	FilterHandler = MakeShared<FTreeFilter>();
-	FilterHandler->SetFilter(SearchFilter.Get());
-	FilterHandler->SetRootItems(&TreeSource, &FilteredTreeSource);
-	FilterHandler->SetGetChildrenDelegate(FTreeFilter::FOnGetChildren::CreateRaw(this, &FPropertyViewerImpl::HandleGetChildren));
-
 	return SAssignNew(SearchBoxWidget, SSearchBox)
 		.HintText(LOCTEXT("SearchHintText", "Search"))
 		.OnTextChanged(this, &FPropertyViewerImpl::HandleSearchChanged);
 }
 
 
-FText FPropertyViewerImpl::SetRawFilterText(const FText& InFilterText)
+void FPropertyViewerImpl::SetRawFilterText(const FText& InFilterText)
+{
+	SetRawFilterTextInternal(InFilterText);
+}
+
+
+FText FPropertyViewerImpl::SetRawFilterTextInternal(const FText& InFilterText)
 {
 	const bool bNewFilteringEnabled = !InFilterText.IsEmpty();
 	FilterHandler->SetIsEnabled(bNewFilteringEnabled);
@@ -840,6 +849,7 @@ TSharedRef<ITableRow> FPropertyViewerImpl::HandleGenerateRow(TSharedPtr<FTreeNod
 								return SNew(SHorizontalBox)
 								+ SHorizontalBox::Slot()
 								.AutoWidth()
+								.VAlign(EVerticalAlignment::VAlign_Center)
 								[
 									ValueWidget.ToSharedRef()
 								]
@@ -960,7 +970,7 @@ void FPropertyViewerImpl::HandleSearchChanged(const FText& InFilterText)
 {
 	if (SearchBoxWidget)
 	{
-		SearchBoxWidget->SetError(SetRawFilterText(InFilterText));
+		SearchBoxWidget->SetError(SetRawFilterTextInternal(InFilterText));
 	}
 }
 
diff --git a/Engine/Source/Runtime/AdvancedWidgets/Private/Widgets/PropertyViewer/SPropertyViewerImpl.h b/Engine/Source/Runtime/AdvancedWidgets/Private/Widgets/PropertyViewer/SPropertyViewerImpl.h
--- a/Engine/Source/Runtime/AdvancedWidgets/Private/Widgets/PropertyViewer/SPropertyViewerImpl.h
+++ b/Engine/Source/Runtime/AdvancedWidgets/Private/Widgets/PropertyViewer/SPropertyViewerImpl.h
@@ -33,9 +33,9 @@ namespace UE::PropertyViewer::Private
  */
 struct FContainer
 {
-	FContainer(SPropertyViewer::FHandle InIdentifier, UObject* InstanceToDisplay);
-	FContainer(SPropertyViewer::FHandle InIdentifier, const UStruct* ClassToDisplay);
-	FContainer(SPropertyViewer::FHandle InIdentifier, const UScriptStruct* Struct, void* Data);
+	FContainer(SPropertyViewer::FHandle InIdentifier, TOptional<FText> DisplayName, const UStruct* ClassToDisplay);
+	FContainer(SPropertyViewer::FHandle InIdentifier, TOptional<FText> DisplayName, UObject* InstanceToDisplay);
+	FContainer(SPropertyViewer::FHandle InIdentifier, TOptional<FText> DisplayName, const UScriptStruct* Struct, void* Data);
 
 public:
 	SPropertyViewer::FHandle GetIdentifier() const
@@ -84,6 +84,11 @@ public:
 		return StructInstance;
 	}
 
+	TOptional<FText> GetDisplayName() const
+	{
+		return DisplayName;
+	}
+
 	bool IsValid() const;
 
 private:
@@ -91,6 +96,7 @@ private:
 	TWeakObjectPtr<const UStruct> Container;
 	TWeakObjectPtr<UObject> ObjectInstance;
 	void* StructInstance = nullptr;
+	TOptional<FText> DisplayName;
 	bool bIsObject = false;
 };
 
@@ -204,11 +210,12 @@ private:
 public:
 	TSharedRef<SWidget> Construct(const SPropertyViewer::FArguments& InArgs);
 
-	void AddContainer(SPropertyViewer::FHandle Identifier, const UStruct* Struct);
-	void AddContainerInstance(SPropertyViewer::FHandle Identifier, UObject* Object);
-	void AddContainerInstance(SPropertyViewer::FHandle Identifier, const UScriptStruct* Struct, void* Data);
+	void AddContainer(SPropertyViewer::FHandle Identifier, TOptional<FText> DisplayName, const UStruct* Struct);
+	void AddContainerInstance(SPropertyViewer::FHandle Identifier, TOptional<FText> DisplayName, UObject* Object);
+	void AddContainerInstance(SPropertyViewer::FHandle Identifier, TOptional<FText> DisplayName, const UScriptStruct* Struct, void* Data);
 	void Remove(SPropertyViewer::FHandle Identifier);
 	void RemoveAll();
+	void SetRawFilterText(const FText& InFilterText);
 
 private:
 	void AddContainerInternal(SPropertyViewer::FHandle Identifier, TSharedPtr<FContainer>& NewContainer);
@@ -217,7 +224,7 @@ private:
 	TSharedRef<SWidget> CreateTree(bool bHasPreWidget, bool bShowPropertyValue, bool bHasPostWidget);
 
 	void HandleSearchChanged(const FText& InFilterText);
-	FText SetRawFilterText(const FText& InFilterText);
+	FText SetRawFilterTextInternal(const FText& InFilterText);
 	void SetHighlightTextRecursive(const TSharedPtr<FTreeNode>& OwnerNode, const FText& HighlightText);
 	void HandleGetFilterStrings(TSharedPtr<FTreeNode> Item, TArray<FString>& OutStrings);
 
diff --git a/Engine/Source/Runtime/AdvancedWidgets/Public/Widgets/PropertyViewer/SPropertyViewer.h b/Engine/Source/Runtime/AdvancedWidgets/Public/Widgets/PropertyViewer/SPropertyViewer.h
--- a/Engine/Source/Runtime/AdvancedWidgets/Public/Widgets/PropertyViewer/SPropertyViewer.h
+++ b/Engine/Source/Runtime/AdvancedWidgets/Public/Widgets/PropertyViewer/SPropertyViewer.h
@@ -95,15 +95,17 @@ public:
 	void Construct(const FArguments& InArgs, const UFunction* Function);
 
 public:
-	FHandle AddContainer(const UScriptStruct* Struct);
-	FHandle AddContainer(const UClass* Class);
-	FHandle AddContainer(const UFunction* Function);
-	FHandle AddInstance(const UScriptStruct* Struct, void* Data);
-	FHandle AddInstance(UObject* ObjectInstance);
+	FHandle AddContainer(const UScriptStruct* Struct, TOptional<FText> DisplayName = TOptional<FText>());
+	FHandle AddContainer(const UClass* Class, TOptional<FText> DisplayName = TOptional<FText>());
+	FHandle AddContainer(const UFunction* Function, TOptional<FText> DisplayName = TOptional<FText>());
+	FHandle AddInstance(const UScriptStruct* Struct, void* Data, TOptional<FText> DisplayName = TOptional<FText>());
+	FHandle AddInstance(UObject* ObjectInstance, TOptional<FText> DisplayName = TOptional<FText>());
 
 	void Remove(FHandle Identifier);
 	void RemoveAll();
 
+	void SetRawFilterText(const FText& InFilterText);
+
 private:
 	void ConstructInternal(const FArguments& InArgs);
 	static FHandle MakeContainerIdentifier();
