PropertyViewer: Fix the tooltip for SFieldName

diff --git a/Engine/Source/Runtime/AdvancedWidgets/Private/Widgets/PropertyViewer/SFieldName.cpp b/Engine/Source/Runtime/AdvancedWidgets/Private/Widgets/PropertyViewer/SFieldName.cpp
--- a/Engine/Source/Runtime/AdvancedWidgets/Private/Widgets/PropertyViewer/SFieldName.cpp
+++ b/Engine/Source/Runtime/AdvancedWidgets/Private/Widgets/PropertyViewer/SFieldName.cpp
@@ -10,6 +10,8 @@
 #include "Widgets/SBoxPanel.h"
 
 
+#define LOCTEXT_NAMESPACE "SFieldName"
+
 namespace UE::PropertyViewer
 {
 
@@ -149,6 +151,7 @@ void SFieldName::Construct(const FArguments& InArgs, const FText& DisplayName, T
 				SAssignNew(NameBlock, STextBlock)
 				.Text(DisplayName)
 				.HighlightText(InArgs._HighlightText)
+				.ToolTipText(GetToolTipText())
 			]
 		];
 	}
@@ -159,12 +162,9 @@ void SFieldName::Construct(const FArguments& InArgs, const FText& DisplayName, T
 			SAssignNew(NameBlock, STextBlock)
 			.Text(DisplayName)
 			.HighlightText(InArgs._HighlightText)
+			.ToolTipText(GetToolTipText())
 		];
 	}
-
-#if WITH_EDITORONLY_DATA
-	SetToolTip(TAttribute<TSharedPtr<IToolTip>>::CreateSP(this, &SFieldName::CreateToolTip));
-#endif
 }
 
 
@@ -176,19 +176,34 @@ void SFieldName::SetHighlightText(TAttribute<FText> InHighlightText)
 	}
 }
 
-TSharedPtr<IToolTip> SFieldName::CreateToolTip() const
+FText SFieldName::GetToolTipText() const
 {
 #if WITH_EDITORONLY_DATA
 	if (FProperty* PropertyPtr = Field.Get<FProperty>())
 	{
-		return FSlateApplicationBase::Get().MakeToolTip(PropertyPtr->GetToolTipText());
+		return FText::FormatOrdered(LOCTEXT("PropertyTooltip", "{0}\nProperty: {1} {2}")
+			, PropertyPtr->GetToolTipText()
+			, FText::FromString(PropertyPtr->GetCPPType())
+			, FText::FromString(PropertyPtr->GetName()));
+	}
+	if (UFunction* FunctionPtr = Field.Get<UFunction>())
+	{
+		const FProperty* ReturnType = FunctionPtr->GetReturnProperty();
+		return FText::FormatOrdered(LOCTEXT("FunctionTooltip", "{0}\nFunction: {1}\nReturns: {2}")
+			, FunctionPtr->GetToolTipText()
+			, FText::FromString(FunctionPtr->GetName())
+			, (ReturnType ? FText::FromString(ReturnType->GetCPPType()) : FText::GetEmpty()));
 	}
 	if (UField* FieldPtr = Field.Get<UField>())
 	{
-		return FSlateApplicationBase::Get().MakeToolTip(FieldPtr->GetToolTipText());
+		return FText::FormatOrdered(LOCTEXT("FieldTooltip", "{0}\nType: {1}")
+			, FieldPtr->GetToolTipText()
+			, FText::FromString(FieldPtr->GetName()));
 	}
 #endif
-	return TSharedPtr<IToolTip>();
+	return FText::GetEmpty();
 }
 
 } //namespace
+
+#undef LOCTEXT_NAMESPACE
\ No newline at end of file
diff --git a/Engine/Source/Runtime/AdvancedWidgets/Public/Widgets/PropertyViewer/SFieldName.h b/Engine/Source/Runtime/AdvancedWidgets/Public/Widgets/PropertyViewer/SFieldName.h
--- a/Engine/Source/Runtime/AdvancedWidgets/Public/Widgets/PropertyViewer/SFieldName.h
+++ b/Engine/Source/Runtime/AdvancedWidgets/Public/Widgets/PropertyViewer/SFieldName.h
@@ -35,7 +35,7 @@ public:
 	void SetHighlightText(TAttribute<FText> InHighlightText);
 
 private:
-	TSharedPtr<IToolTip> CreateToolTip() const;
+	FText GetToolTipText() const;
 	void Construct(const FArguments& InArgs, const FText& DisplayName, TSharedPtr<SWidget> Icon);
 
 	TSharedPtr<STextBlock> NameBlock;
