SPropertyViewer delegates now pass an array of FFieldVariants instead of just the leaf-most field.

SPropertyViewer now supports an OnDoubleClicked delegate.

diff --git a/Engine/Source/Runtime/AdvancedWidgets/Private/Widgets/PropertyViewer/SPropertyViewerImpl.cpp b/Engine/Source/Runtime/AdvancedWidgets/Private/Widgets/PropertyViewer/SPropertyViewerImpl.cpp
--- a/Engine/Source/Runtime/AdvancedWidgets/Private/Widgets/PropertyViewer/SPropertyViewerImpl.cpp
+++ b/Engine/Source/Runtime/AdvancedWidgets/Private/Widgets/PropertyViewer/SPropertyViewerImpl.cpp
@@ -174,6 +174,22 @@ FPropertyPath FTreeNode::GetPropertyPath() const
 	return FPropertyPath();
 }
 
+TArray<FFieldVariant> FTreeNode::GetFieldPath() const
+{
+	TArray<FFieldVariant> Fields;
+
+	const FTreeNode* CurrentNode = this;
+	while (CurrentNode)
+	{ 
+		if (FFieldVariant Field = CurrentNode->GetField())
+		{
+			Fields.Insert(Field, 0);
+		}
+		
+		CurrentNode = CurrentNode->ParentNode.Pin().Get();
+	}
+	return Fields;
+}
 
 TSharedPtr<FContainer> FTreeNode::GetOwnerContainer() const
 {
@@ -380,6 +396,7 @@ FPropertyViewerImpl::FPropertyViewerImpl(const SPropertyViewer::FArguments& InAr
 	OnGetPreSlot = InArgs._OnGetPreSlot;
 	OnGetPostSlot = InArgs._OnGetPostSlot;
 	OnSelectionChanged = InArgs._OnSelectionChanged;
+	OnDoubleClicked = InArgs._OnDoubleClicked;
 	PropertyVisibility = InArgs._PropertyVisibility;
 	bSanitizeName = InArgs._bSanitizeName;
 	bShowFieldIcon = InArgs._bShowFieldIcon;
@@ -683,6 +700,7 @@ TSharedRef<SWidget> FPropertyViewerImpl::CreateTree(bool bHasPreWidget, bool bSh
 			.OnGenerateRow(this, &FPropertyViewerImpl::HandleGenerateRow)
 			.OnSelectionChanged(this, &FPropertyViewerImpl::HandleSelectionChanged)
 			.OnContextMenuOpening(this, &FPropertyViewerImpl::HandleContextMenuOpening)
+			.OnMouseButtonDoubleClick(this, &FPropertyViewerImpl::HandleDoubleClick)
 			.HeaderRow(HeaderRowWidget);
 		FilterHandler->SetTreeView(TreeWidget.Get());
 	}
@@ -696,6 +714,7 @@ TSharedRef<SWidget> FPropertyViewerImpl::CreateTree(bool bHasPreWidget, bool bSh
 			.OnGenerateRow(this, &FPropertyViewerImpl::HandleGenerateRow)
 			.OnSelectionChanged(this, &FPropertyViewerImpl::HandleSelectionChanged)
 			.OnContextMenuOpening(this, &FPropertyViewerImpl::HandleContextMenuOpening)
+			.OnMouseButtonDoubleClick(this, &FPropertyViewerImpl::HandleDoubleClick)
 			.HeaderRow(HeaderRowWidget);
 	}
 
@@ -872,14 +891,9 @@ TSharedRef<ITableRow> FPropertyViewerImpl::HandleGenerateRow(TSharedPtr<FTreeNod
 				{
 					SPropertyViewer::FGetFieldWidget& OnGetWidget = (ColumnName == ColumnName_FieldPreWidget) ? PropertyViewOwnerPin->OnGetPreSlot : PropertyViewOwnerPin->OnGetPostSlot;
 					TSharedPtr<SWidget> PreWidget;
-					if (TSharedPtr<FContainer> ContainerPin = ItemPin->GetContainer())
-					{
-						PreWidget = OnGetWidget.Execute(ContainerPin->GetIdentifier(), FFieldVariant(ContainerPin->GetStruct()));
-					}
-					else if (FFieldVariant FieldVariant = ItemPin->GetField())
-					{
-						PreWidget = OnGetWidget.Execute(SPropertyViewer::FHandle(), FieldVariant);
-					}
+
+					TSharedPtr<FContainer> OwnerContainer = ItemPin->GetOwnerContainer();
+					PreWidget = OnGetWidget.Execute(OwnerContainer.IsValid() ? OwnerContainer->GetIdentifier() : SPropertyViewer::FHandle(), ItemPin->GetFieldPath());
 
 					if (PreWidget)
 					{
@@ -936,14 +950,8 @@ TSharedPtr<SWidget> FPropertyViewerImpl::HandleContextMenuOpening()
 		TArray<TSharedPtr<FTreeNode>> Items = TreeWidget->GetSelectedItems();
 		if (Items.Num() == 1 && Items[0].IsValid())
 		{
-			if (TSharedPtr<FContainer> ContainerPin = Items[0]->GetContainer())
-			{
-				return OnContextMenuOpening.Execute(ContainerPin->GetIdentifier(), FFieldVariant(ContainerPin->GetStruct()));
-			}
-			else if (FFieldVariant FieldVariant = Items[0]->GetField())
-			{
-				return OnContextMenuOpening.Execute(SPropertyViewer::FHandle(), FieldVariant);
-			}
+			TSharedPtr<FContainer> OwnerContainer = Items[0]->GetOwnerContainer();
+			return OnContextMenuOpening.Execute(OwnerContainer.IsValid() ? OwnerContainer->GetIdentifier() : SPropertyViewer::FHandle(), Items[0]->GetFieldPath());
 		}
 	}
 	return TSharedPtr<SWidget>();
@@ -954,13 +962,30 @@ void FPropertyViewerImpl::HandleSelectionChanged(TSharedPtr<FTreeNode> Item, ESe
 {
 	if (OnSelectionChanged.IsBound())
 	{
-		if (TSharedPtr<FContainer> ContainerPin = Item->GetContainer())
+		if (Item.IsValid())
+		{
+			TSharedPtr<FContainer> OwnerContainer = Item->GetOwnerContainer();
+			OnSelectionChanged.Execute(OwnerContainer.IsValid() ? OwnerContainer->GetIdentifier() : SPropertyViewer::FHandle(), Item->GetFieldPath(), SelectionType);
+		}
+		else
+		{
+			OnSelectionChanged.Execute(SPropertyViewer::FHandle(), TArray<FFieldVariant>(), SelectionType);
+		}
+	}
+}
+
+void FPropertyViewerImpl::HandleDoubleClick(TSharedPtr<FTreeNode> Item)
+{
+	if (OnDoubleClicked.IsBound())
+	{
+		if (Item.IsValid())
 		{
-			OnSelectionChanged.Execute(ContainerPin->GetIdentifier(), FFieldVariant(ContainerPin->GetStruct()), SelectionType);
+			TSharedPtr<FContainer> OwnerContainer = Item->GetOwnerContainer();
+			OnDoubleClicked.Execute(OwnerContainer.IsValid() ? OwnerContainer->GetIdentifier() : SPropertyViewer::FHandle(), Item->GetFieldPath());
 		}
-		else if (FFieldVariant FieldVariant = Item->GetField())
+		else
 		{
-			OnSelectionChanged.Execute(SPropertyViewer::FHandle(), FieldVariant, SelectionType);
+			OnDoubleClicked.Execute(SPropertyViewer::FHandle(), TArray<FFieldVariant>());
 		}
 	}
 }
diff --git a/Engine/Source/Runtime/AdvancedWidgets/Private/Widgets/PropertyViewer/SPropertyViewerImpl.h b/Engine/Source/Runtime/AdvancedWidgets/Private/Widgets/PropertyViewer/SPropertyViewerImpl.h
--- a/Engine/Source/Runtime/AdvancedWidgets/Private/Widgets/PropertyViewer/SPropertyViewerImpl.h
+++ b/Engine/Source/Runtime/AdvancedWidgets/Private/Widgets/PropertyViewer/SPropertyViewerImpl.h
@@ -148,6 +148,7 @@ public:
 	}
 
 	FPropertyPath GetPropertyPath() const;
+	TArray<FFieldVariant> GetFieldPath() const;
 
 	TSharedPtr<FContainer> GetOwnerContainer() const;
 
@@ -195,6 +196,7 @@ private:
 	SPropertyViewer::FGetFieldWidget OnGetPostSlot;
 	SPropertyViewer::FOnContextMenuOpening OnContextMenuOpening;
 	SPropertyViewer::FOnSelectionChanged OnSelectionChanged;
+	SPropertyViewer::FOnDoubleClicked OnDoubleClicked;
 	SPropertyViewer::EPropertyVisibility PropertyVisibility;
 	bool bSanitizeName = false;
 	bool bShowFieldIcon = false;
@@ -233,6 +235,7 @@ private:
 
 	TSharedPtr<SWidget> HandleContextMenuOpening();
 	void HandleSelectionChanged(TSharedPtr<FTreeNode> Item, ESelectInfo::Type SelectionType);
+	void HandleDoubleClick(TSharedPtr<FTreeNode> Item);
 
 #if WITH_EDITOR
 	void HandleBlueprintCompiled();
diff --git a/Engine/Source/Runtime/AdvancedWidgets/Public/Widgets/PropertyViewer/SPropertyViewer.h b/Engine/Source/Runtime/AdvancedWidgets/Public/Widgets/PropertyViewer/SPropertyViewer.h
--- a/Engine/Source/Runtime/AdvancedWidgets/Public/Widgets/PropertyViewer/SPropertyViewer.h
+++ b/Engine/Source/Runtime/AdvancedWidgets/Public/Widgets/PropertyViewer/SPropertyViewer.h
@@ -51,9 +51,10 @@ public:
 	};
 
 public:
-	DECLARE_DELEGATE_RetVal_TwoParams(TSharedPtr<SWidget>, FGetFieldWidget, FHandle, const FFieldVariant);
-	DECLARE_DELEGATE_RetVal_TwoParams(TSharedPtr<SWidget>, FOnContextMenuOpening, FHandle, const FFieldVariant);
-	DECLARE_DELEGATE_ThreeParams(FOnSelectionChanged, FHandle, const FFieldVariant, ESelectInfo::Type);
+	DECLARE_DELEGATE_RetVal_TwoParams(TSharedPtr<SWidget>, FGetFieldWidget, FHandle, TArrayView<const FFieldVariant>);
+	DECLARE_DELEGATE_RetVal_TwoParams(TSharedPtr<SWidget>, FOnContextMenuOpening, FHandle, TArrayView<const FFieldVariant>);
+	DECLARE_DELEGATE_ThreeParams(FOnSelectionChanged, FHandle, TArrayView<const FFieldVariant>, ESelectInfo::Type);
+	DECLARE_DELEGATE_TwoParams(FOnDoubleClicked, FHandle, TArrayView<const FFieldVariant>);
 
 	SLATE_BEGIN_ARGS(SPropertyViewer)
 	{}
@@ -85,6 +86,8 @@ public:
 		SLATE_EVENT(FOnContextMenuOpening, OnContextMenuOpening)
 		/** Delegate to invoke when selection changes. */
 		SLATE_EVENT(FOnSelectionChanged, OnSelectionChanged)
+		/** Delegate to invoke when an item is double-clicked. */
+		SLATE_EVENT(FOnDoubleClicked, OnDoubleClicked);
 	SLATE_END_ARGS()
 
 	void Construct(const FArguments& InArgs);
