Fix compiler issue with namespace.
#rxn
#issue 195297

diff --git a/Engine/Source/Runtime/AdvancedWidgets/Public/Widgets/PropertyViewer/SPropertyViewer.h b/Engine/Source/Runtime/AdvancedWidgets/Public/Widgets/PropertyViewer/SPropertyViewer.h
--- a/Engine/Source/Runtime/AdvancedWidgets/Public/Widgets/PropertyViewer/SPropertyViewer.h
+++ b/Engine/Source/Runtime/AdvancedWidgets/Public/Widgets/PropertyViewer/SPropertyViewer.h
@@ -43,7 +43,7 @@ public:
 		}
 		friend uint32 GetTypeHash(const FHandle& Other)
 		{
-			return GetTypeHash(Other.Id);
+			return ::GetTypeHash(Other.Id);
 		}
 	};
 
