Fix compiler issue with namespace.
#rxn
#issue 195297

diff --git a/Engine/Source/Runtime/AdvancedWidgets/Public/Widgets/PropertyViewer/SPropertyViewer.h b/Engine/Source/Runtime/AdvancedWidgets/Public/Widgets/PropertyViewer/SPropertyViewer.h
--- a/Engine/Source/Runtime/AdvancedWidgets/Public/Widgets/PropertyViewer/SPropertyViewer.h
+++ b/Engine/Source/Runtime/AdvancedWidgets/Public/Widgets/PropertyViewer/SPropertyViewer.h
@@ -2,6 +2,7 @@
 
 #pragma once
 
+#include "Templates/TypeHash.h"
 #include "Widgets/SCompoundWidget.h"
 #include "UObject/Class.h"
 #include "UObject/Field.h"
