Exposed SelectionMode on SPropertyViewer.

#rnx

diff --git a/Engine/Source/Runtime/AdvancedWidgets/Private/Widgets/PropertyViewer/SPropertyViewerImpl.cpp b/Engine/Source/Runtime/AdvancedWidgets/Private/Widgets/PropertyViewer/SPropertyViewerImpl.cpp
--- a/Engine/Source/Runtime/AdvancedWidgets/Private/Widgets/PropertyViewer/SPropertyViewerImpl.cpp
+++ b/Engine/Source/Runtime/AdvancedWidgets/Private/Widgets/PropertyViewer/SPropertyViewerImpl.cpp
@@ -451,10 +451,12 @@ TSharedRef<SWidget> FPropertyViewerImpl::Construct(const SPropertyViewer::FArgum
 				.AutoWidth()
 				.HAlign(HAlign_Left)
 				.VAlign(VAlign_Center)
+				.Padding(0, 0, 4, 0)
 				[
 					InArgs._SearchBoxPreSlot.Widget
 				];
 		}
+
 		if (InArgs._bShowSearchBox)
 		{
 			SearchBox->AddSlot()
@@ -472,12 +474,14 @@ TSharedRef<SWidget> FPropertyViewerImpl::Construct(const SPropertyViewer::FArgum
 					SNullWidget::NullWidget
 				];
 		}
+
 		if (InArgs._SearchBoxPostSlot.Widget != SNullWidget::NullWidget)
 		{
 			SearchBox->AddSlot()
 				.AutoWidth()
 				.HAlign(HAlign_Right)
 				.VAlign(VAlign_Center)
+				.Padding(4, 0, 0, 0)
 				[
 					InArgs._SearchBoxPostSlot.Widget
 				];
@@ -488,7 +492,7 @@ TSharedRef<SWidget> FPropertyViewerImpl::Construct(const SPropertyViewer::FArgum
 		.BorderImage(FAppStyle::Get().GetBrush("Brushes.Recessed"))
 		.Padding(0)
 		[
-			CreateTree(OnGetPreSlot.IsBound(), PropertyVisibility != SPropertyViewer::EPropertyVisibility::Hidden, OnGetPostSlot.IsBound())
+			CreateTree(OnGetPreSlot.IsBound(), PropertyVisibility != SPropertyViewer::EPropertyVisibility::Hidden, OnGetPostSlot.IsBound(), InArgs._SelectionMode)
 		];
 
 	if (SearchBox)
@@ -683,7 +687,7 @@ FText FPropertyViewerImpl::SetRawFilterTextInternal(const FText& InFilterText)
 }
 
 
-TSharedRef<SWidget> FPropertyViewerImpl::CreateTree(bool bHasPreWidget, bool bShowPropertyValue, bool bHasPostWidget)
+TSharedRef<SWidget> FPropertyViewerImpl::CreateTree(bool bHasPreWidget, bool bShowPropertyValue, bool bHasPostWidget, ESelectionMode::Type SelectionMode)
 {
 	TSharedPtr<SHeaderRow> HeaderRowWidget;
 	if (bHasPreWidget || bShowPropertyValue || bHasPostWidget)
@@ -734,7 +738,7 @@ TSharedRef<SWidget> FPropertyViewerImpl::CreateTree(bool bHasPreWidget, bool bSh
 		SAssignNew(TreeWidget, STreeView<TSharedPtr<FTreeNode>>)
 			.ItemHeight(1.0f)
 			.TreeItemsSource(&FilteredTreeSource)
-			.SelectionMode(ESelectionMode::Single)
+			.SelectionMode(SelectionMode)
 			.OnGetChildren(FilterHandler.ToSharedRef(), &FTreeFilter::OnGetFilteredChildren)
 			.OnGenerateRow(this, &FPropertyViewerImpl::HandleGenerateRow)
 			.OnSelectionChanged(this, &FPropertyViewerImpl::HandleSelectionChanged)
@@ -748,7 +752,7 @@ TSharedRef<SWidget> FPropertyViewerImpl::CreateTree(bool bHasPreWidget, bool bSh
 		SAssignNew(TreeWidget, STreeView<TSharedPtr<FTreeNode>>)
 			.ItemHeight(1.0f)
 			.TreeItemsSource(&TreeSource)
-			.SelectionMode(ESelectionMode::Single)
+			.SelectionMode(SelectionMode)
 			.OnGetChildren(this, &FPropertyViewerImpl::HandleGetChildren)
 			.OnGenerateRow(this, &FPropertyViewerImpl::HandleGenerateRow)
 			.OnSelectionChanged(this, &FPropertyViewerImpl::HandleSelectionChanged)
diff --git a/Engine/Source/Runtime/AdvancedWidgets/Private/Widgets/PropertyViewer/SPropertyViewerImpl.h b/Engine/Source/Runtime/AdvancedWidgets/Private/Widgets/PropertyViewer/SPropertyViewerImpl.h
--- a/Engine/Source/Runtime/AdvancedWidgets/Private/Widgets/PropertyViewer/SPropertyViewerImpl.h
+++ b/Engine/Source/Runtime/AdvancedWidgets/Private/Widgets/PropertyViewer/SPropertyViewerImpl.h
@@ -225,7 +225,7 @@ private:
 	void AddContainerInternal(SPropertyViewer::FHandle Identifier, TSharedPtr<FContainer>& NewContainer);
 
 	TSharedRef<SWidget> CreateSearch();
-	TSharedRef<SWidget> CreateTree(bool bHasPreWidget, bool bShowPropertyValue, bool bHasPostWidget);
+	TSharedRef<SWidget> CreateTree(bool bHasPreWidget, bool bShowPropertyValue, bool bHasPostWidget, ESelectionMode::Type SelectionMode);
 
 	void HandleSearchChanged(const FText& InFilterText);
 	FText SetRawFilterTextInternal(const FText& InFilterText);
diff --git a/Engine/Source/Runtime/AdvancedWidgets/Public/Widgets/PropertyViewer/SPropertyViewer.h b/Engine/Source/Runtime/AdvancedWidgets/Public/Widgets/PropertyViewer/SPropertyViewer.h
--- a/Engine/Source/Runtime/AdvancedWidgets/Public/Widgets/PropertyViewer/SPropertyViewer.h
+++ b/Engine/Source/Runtime/AdvancedWidgets/Public/Widgets/PropertyViewer/SPropertyViewer.h
@@ -2,6 +2,7 @@
 
 #pragma once
 
+#include "Framework/Views/ITypedTableView.h"
 #include "Templates/TypeHash.h"
 #include "Widgets/SCompoundWidget.h"
 #include "UObject/Class.h"
@@ -81,6 +82,8 @@ public:
 		SLATE_ARGUMENT_DEFAULT(bool, bSortChildNode) = false;
 		/** Show a search box. */
 		SLATE_ARGUMENT_DEFAULT(bool, bShowSearchBox) = false;
+		/** Is selection enabled. */ 
+		SLATE_ARGUMENT_DEFAULT(ESelectionMode::Type, SelectionMode) = ESelectionMode::Single;
 		/** Which properties/functions to show. FFieldIterator_BlueprintVisible is the default. */
 		SLATE_ARGUMENT_DEFAULT(IFieldIterator*, FieldIterator) = nullptr;
 		/** Which properties/functions that allow expansion. FFieldIterator_NoExpand is the default. */
