PropertyViewer: Use UObject instance to expand the field tree. Cache node that expanded from instance. Test each frame if the instance changes.

diff --git a/Engine/Source/Runtime/AdvancedWidgets/Private/Widgets/PropertyViewer/SPropertyViewer.cpp b/Engine/Source/Runtime/AdvancedWidgets/Private/Widgets/PropertyViewer/SPropertyViewer.cpp
--- a/Engine/Source/Runtime/AdvancedWidgets/Private/Widgets/PropertyViewer/SPropertyViewer.cpp
+++ b/Engine/Source/Runtime/AdvancedWidgets/Private/Widgets/PropertyViewer/SPropertyViewer.cpp
@@ -53,6 +53,13 @@ void SPropertyViewer::Construct(const FArguments& InArgs, const UFunction* Funct
 }
 
 
+void SPropertyViewer::Tick(const FGeometry& AllottedGeometry, const double InCurrentTime, const float InDeltaTime)
+{
+	Super::Tick(AllottedGeometry, InCurrentTime, InDeltaTime);
+	Implementation->Tick();
+}
+
+
 SPropertyViewer::FHandle SPropertyViewer::AddContainer(const UScriptStruct* Struct, TOptional<FText> DisplayName)
 {
 	SPropertyViewer::FHandle Result = MakeContainerIdentifier();
diff --git a/Engine/Source/Runtime/AdvancedWidgets/Private/Widgets/PropertyViewer/SPropertyViewerImpl.cpp b/Engine/Source/Runtime/AdvancedWidgets/Private/Widgets/PropertyViewer/SPropertyViewerImpl.cpp
--- a/Engine/Source/Runtime/AdvancedWidgets/Private/Widgets/PropertyViewer/SPropertyViewerImpl.cpp
+++ b/Engine/Source/Runtime/AdvancedWidgets/Private/Widgets/PropertyViewer/SPropertyViewerImpl.cpp
@@ -134,7 +134,7 @@ FPropertyPath FTreeNode::GetPropertyPath() const
 	{
 		if (CurrentNode->Property)
 		{
-			Properties.Insert(Property, 0);
+			Properties.Insert(CurrentNode->Property, 0);
 		}
 
 		TSharedPtr<FContainer> ContainerPin = CurrentNode->Container.Pin();
@@ -248,13 +248,15 @@ void FTreeNode::GetFilterStrings(TArray<FString>& OutStrings) const
 }
 
 
-void FTreeNode::BuildChildNodes(IFieldIterator& FieldIterator, IFieldExpander& FieldExpander, bool bSortChildNode)
+TArray<FTreeNode::FNodeReference> FTreeNode::BuildChildNodes(IFieldIterator& FieldIterator, IFieldExpander& FieldExpander, bool bSortChildNode)
 {
-	BuildChildNodesRecursive(FieldIterator, FieldExpander, bSortChildNode, 2);
+	TArray<FNodeReference> TickReferences;
+	BuildChildNodesRecursive(FieldIterator, FieldExpander, bSortChildNode, 2, TickReferences);
+	return TickReferences;
 }
 
 
-void FTreeNode::BuildChildNodesRecursive(IFieldIterator& FieldIterator, IFieldExpander& FieldExpander, bool bSortChildNode, int32 RecursiveCount)
+void FTreeNode::BuildChildNodesRecursive(IFieldIterator& FieldIterator, IFieldExpander& FieldExpander, bool bSortChildNode, int32 RecursiveCount, TArray<FNodeReference>& OutTickReference)
 {
 	if (RecursiveCount <= 0)
 	{
@@ -269,12 +271,14 @@ void FTreeNode::BuildChildNodesRecursive(IFieldIterator& FieldIterator, IFieldEx
 	{
 		if (const FStructProperty* StructProperty = CastField<const FStructProperty>(Property))
 		{
-			ChildStructType = StructProperty->Struct;
+			if (FieldExpander.CanExpandScriptStruct(StructProperty))
+			{
+				ChildStructType = StructProperty->Struct;
+			}
 		}
 		else if (const FObjectPropertyBase* ObjectProperty = CastField<const FObjectPropertyBase>(Property))
 		{
-			// If the container is an instance and the object is nullptr, do not expand this object.
-			bool bIsNullptr = false;
+			UObject* Instance = nullptr;
 			if (TSharedPtr<FContainer> OwnerContainer = GetOwnerContainer())
 			{
 				if (OwnerContainer->IsInstance())
@@ -282,16 +286,19 @@ void FTreeNode::BuildChildNodesRecursive(IFieldIterator& FieldIterator, IFieldEx
 					const FPropertyPath PropertyPath = GetPropertyPath();
 					if (const void* ContainerPtr = PropertyPath.GetContainerPtr())
 					{
-						bIsNullptr = ObjectProperty->GetObjectPropertyValue_InContainer(ContainerPtr) == nullptr;
+						Instance = ObjectProperty->GetObjectPropertyValue_InContainer(ContainerPtr);
 					}
 				}
 			}
 
-			if (!bIsNullptr && FieldExpander.CanExpandObject(ObjectProperty->PropertyClass))
+			TOptional<const UClass*> ClassToExpand = FieldExpander.CanExpandObject(ObjectProperty, Instance);
+			if (ClassToExpand.IsSet())
 			{
-				//todo expand the object instance, not the object class
-				//todo we want to confirm that that thing didn't change frame to frame
-				ChildStructType = ObjectProperty->PropertyClass;
+				if (Instance)
+				{
+					OutTickReference.Emplace(Instance, AsShared());
+				}
+				ChildStructType = ClassToExpand.GetValue();
 			}
 		}
 	}
@@ -314,12 +321,12 @@ void FTreeNode::BuildChildNodesRecursive(IFieldIterator& FieldIterator, IFieldEx
 			if (const FProperty* PropertyIt = FieldIt.Get<FProperty>())
 			{
 				TSharedPtr<FTreeNode> Node = MakeField(AsShared(), PropertyIt, TOptional<FText>());
-				Node->BuildChildNodesRecursive(FieldIterator, FieldExpander, bSortChildNode, RecursiveCount);
+				Node->BuildChildNodesRecursive(FieldIterator, FieldExpander, bSortChildNode, RecursiveCount, OutTickReference);
 			}
 			if (const UFunction* FunctionIt = FieldIt.Get<UFunction>())
 			{
 				TSharedPtr<FTreeNode> Node = MakeField(AsShared(), FunctionIt, TOptional<FText>());
-				Node->BuildChildNodesRecursive(FieldIterator, FieldExpander, bSortChildNode, RecursiveCount);
+				Node->BuildChildNodesRecursive(FieldIterator, FieldExpander, bSortChildNode, RecursiveCount, OutTickReference);
 			}
 		}
 
@@ -388,7 +395,10 @@ FPropertyViewerImpl::FPropertyViewerImpl(const SPropertyViewer::FArguments& InAr
 	}
 	if (FieldExpander == nullptr)
 	{
-		FieldExpander = new FFieldExpander_NoExpand();
+		FFieldExpander_Default* DefaultFieldExpander = new FFieldExpander_Default();
+		DefaultFieldExpander->SetExpandObject(FFieldExpander_Default::EObjectExpandFlag::None);
+		DefaultFieldExpander->SetExpandScriptStruct(true);
+		FieldExpander = DefaultFieldExpander;
 		bOwnFieldExpander = true;
 	}
 
@@ -517,6 +527,52 @@ TSharedRef<SWidget> FPropertyViewerImpl::Construct(const SPropertyViewer::FArgum
 }
 
 
+void FPropertyViewerImpl::Tick()
+{
+	auto RemoveChild = [](const TSharedPtr<STreeView<TSharedPtr<FTreeNode>>>& InTreeWidget, const TSharedPtr<FTreeNode>& InNodePin)
+		{
+			if (TSharedPtr<FTreeNode> ParentNode = InNodePin->GetParentNode())
+			{
+				ParentNode->RemoveChild();
+			}
+			InTreeWidget->RequestListRefresh();
+		};
+
+	// If the object instance is not the same as the previous frame. Rebuild the tree.
+	for (int32 Index = NodeReferences.Num() - 1; Index >= 0; --Index)
+	{
+		UObject* PreviousPtr = NodeReferences[Index].Previous.Get();
+		TSharedPtr<FTreeNode> NodePin = NodeReferences[Index].Node.Pin();
+		if (PreviousPtr && NodePin)
+		{
+			UObject* Instance = nullptr;
+			if (FObjectPropertyBase* ObjectProperty = CastField<FObjectPropertyBase>(NodePin->GetField().Get<FProperty>()))
+			{
+				const FPropertyPath PropertyPath = NodePin->GetPropertyPath();
+				if (const void* ContainerPtr = PropertyPath.GetContainerPtr())
+				{
+					Instance = ObjectProperty->GetObjectPropertyValue_InContainer(ContainerPtr);
+				}
+			}
+
+			if (Instance != PreviousPtr)
+			{
+				RemoveChild(TreeWidget, NodePin);
+				NodeReferences.RemoveAtSwap(Index);
+			}
+		}
+		else
+		{
+			if (NodePin)
+			{
+				RemoveChild(TreeWidget, NodePin);
+			}
+			NodeReferences.RemoveAtSwap(Index);
+		}
+	}
+}
+
+
 void FPropertyViewerImpl::AddContainer(SPropertyViewer::FHandle Identifier, TOptional<FText> DisplayName, const UStruct* Struct)
 {
 	TSharedPtr<FContainer> NewContainer = MakeShared<FContainer>(Identifier, DisplayName, Struct);
@@ -547,7 +603,7 @@ void FPropertyViewerImpl::AddContainerInstance(SPropertyViewer::FHandle Identifi
 void FPropertyViewerImpl::AddContainerInternal(SPropertyViewer::FHandle Identifier, TSharedPtr<FContainer>& NewContainer)
 {
 	TSharedPtr<FTreeNode> NewNode = FTreeNode::MakeContainer(NewContainer, NewContainer->GetDisplayName());
-	NewNode->BuildChildNodes(*FieldIterator, *FieldExpander, bSortChildNode);
+	NodeReferences.Append(NewNode->BuildChildNodes(*FieldIterator, *FieldExpander, bSortChildNode));
 	TreeSource.Add(NewNode);
 
 	if (TreeWidget)
@@ -983,7 +1039,7 @@ void FPropertyViewerImpl::HandleGetChildren(TSharedPtr<FTreeNode> InParent, TArr
 		// Do not build when filtering (only search in what it's already been built)
 		if (FilterHandler == nullptr || !FilterHandler->GetIsEnabled())
 		{
-			InParent->BuildChildNodes(*FieldIterator, *FieldExpander, bSortChildNode);
+			NodeReferences.Append(InParent->BuildChildNodes(*FieldIterator, *FieldExpander, bSortChildNode));
 		}
 	}
 	OutChildren = InParent->ChildNodes;
diff --git a/Engine/Source/Runtime/AdvancedWidgets/Private/Widgets/PropertyViewer/SPropertyViewerImpl.h b/Engine/Source/Runtime/AdvancedWidgets/Private/Widgets/PropertyViewer/SPropertyViewerImpl.h
--- a/Engine/Source/Runtime/AdvancedWidgets/Private/Widgets/PropertyViewer/SPropertyViewerImpl.h
+++ b/Engine/Source/Runtime/AdvancedWidgets/Private/Widgets/PropertyViewer/SPropertyViewerImpl.h
@@ -5,8 +5,9 @@
 #include "Framework/PropertyViewer/PropertyPath.h"
 #include "Widgets/PropertyViewer/SPropertyViewer.h"
 #include "Types/SlateEnums.h"
-#include "UObject/WeakObjectPtrTemplates.h"
 #include "UObject/Field.h"
+#include "UObject/Object.h"
+#include "UObject/WeakObjectPtrTemplates.h"
 
 class ITableRow;
 class SSearchBox;
@@ -147,6 +148,11 @@ public:
 		return Property != nullptr ? FFieldVariant(Property) : Function.Get() ? FFieldVariant(Function.Get()) : FFieldVariant();
 	}
 
+	TSharedPtr<FTreeNode> GetParentNode() const
+	{
+		return ParentNode.Pin();
+	}
+
 	FPropertyPath GetPropertyPath() const;
 	TArray<FFieldVariant> GetFieldPath() const;
 
@@ -159,10 +165,27 @@ public:
 
 	void GetFilterStrings(TArray<FString>& OutStrings) const;
 
-	void BuildChildNodes(IFieldIterator& FieldIterator, IFieldExpander& FieldExpander, bool bSortChildNode);
+	struct FNodeReference
+	{
+		TWeakObjectPtr<UObject> Previous;
+		TWeakPtr<FTreeNode> Node;
+
+		FNodeReference(UObject* InPrevious, const TWeakPtr<FTreeNode>& InNode)
+			: Previous(InPrevious)
+			, Node(InNode)
+		{}
+	};
+
+	// Returns the object node that was build
+	TArray<FNodeReference> BuildChildNodes(IFieldIterator& FieldIterator, IFieldExpander& FieldExpander, bool bSortChildNode);
+	void RemoveChild()
+	{
+		ChildNodes.Reset();
+		bChildGenerated = false;
+	}
 
 private:
-	void BuildChildNodesRecursive(IFieldIterator& FieldIterator, IFieldExpander& FieldExpander, bool bSortChildNode, int32 RecursiveCount);
+	void BuildChildNodesRecursive(IFieldIterator& FieldIterator, IFieldExpander& FieldExpander, bool bSortChildNode, int32 RecursiveCount, TArray<FNodeReference>& OutTickReference);
 
 	static bool Sort(const TSharedPtr<FTreeNode>& NodeA, const TSharedPtr<FTreeNode>& NodeB);
 };
@@ -181,6 +204,7 @@ public:
 
 private:
 	TArray<TSharedPtr<FContainer>> Containers;
+	TArray<FTreeNode::FNodeReference> NodeReferences;
 
 	TSharedPtr<SSearchBox> SearchBoxWidget;
 	TSharedPtr<STreeView<TSharedPtr<FTreeNode>>> TreeWidget;
@@ -212,6 +236,7 @@ private:
 
 public:
 	TSharedRef<SWidget> Construct(const SPropertyViewer::FArguments& InArgs);
+	void Tick();
 
 	void AddContainer(SPropertyViewer::FHandle Identifier, TOptional<FText> DisplayName, const UStruct* Struct);
 	void AddContainerInstance(SPropertyViewer::FHandle Identifier, TOptional<FText> DisplayName, UObject* Object);
diff --git a/Engine/Source/Runtime/AdvancedWidgets/Public/Framework/PropertyViewer/IFieldExpander.h b/Engine/Source/Runtime/AdvancedWidgets/Public/Framework/PropertyViewer/IFieldExpander.h
--- a/Engine/Source/Runtime/AdvancedWidgets/Public/Framework/PropertyViewer/IFieldExpander.h
+++ b/Engine/Source/Runtime/AdvancedWidgets/Public/Framework/PropertyViewer/IFieldExpander.h
@@ -3,8 +3,12 @@
 #pragma once
 
 #include "CoreMinimal.h"
+#include "UObject/Object.h"
+#include "UObject/UnrealType.h"
 
+class FObjectPropertyBase;
 class FProperty;
+class FStructProperty;
 class UFunction;
 class UStruct;
 
@@ -15,15 +19,63 @@ namespace UE::PropertyViewer
 class ADVANCEDWIDGETS_API IFieldExpander
 {
 public:
-	virtual bool CanExpandObject(const UClass*) const = 0;
+	/** @return the class that should be expended for this instance. The instance can be null. */
+	virtual TOptional<const UClass*> CanExpandObject(const FObjectPropertyBase* Property, const UObject* Instance) const = 0;
+	/** @return true if the struct property can be expended. */
+	virtual bool CanExpandScriptStruct(const FStructProperty* StructProperty) const = 0;
 	virtual bool CanExpandFunction(const UFunction*) const = 0;
 	virtual ~IFieldExpander() = default;
 };
 
-class ADVANCEDWIDGETS_API FFieldExpander_NoExpand : public IFieldExpander
+class ADVANCEDWIDGETS_API FFieldExpander_Default : public IFieldExpander
 {
-	virtual bool CanExpandObject(const UClass*) const override { return true; }
-	virtual bool CanExpandFunction(const UFunction*) const override { return false; }
+public:
+	FFieldExpander_Default() = default;
+
+	enum class EObjectExpandFlag : uint8
+	{
+		None = 0,
+		UseInstanceClass = 1 << 1,
+		RequireValidInstance = 1 << 2,
+	};
+
+	virtual TOptional<const UClass*> CanExpandObject(const FObjectPropertyBase* Property, const UObject* Instance) const override
+	{
+		const bool bUseInstanceClass = EnumHasAllFlags(ExpandObject, EObjectExpandFlag::UseInstanceClass);
+		const bool bOnlyIfInstanceValid = EnumHasAllFlags(ExpandObject, EObjectExpandFlag::RequireValidInstance);
+		if (bOnlyIfInstanceValid && !Instance)
+		{
+			return TOptional<const UClass*>();
+		}
+		if (bUseInstanceClass && Instance)
+		{
+			return Instance->GetClass();
+		}
+		return Property->PropertyClass;
+	}
+	virtual bool CanExpandScriptStruct(const FStructProperty*) const override { return bExpandScriptStruct; }
+	virtual bool CanExpandFunction(const UFunction*) const override { return bExpandFunction; }
+
+	void SetExpandObject(EObjectExpandFlag InExpandObject)
+	{
+		ExpandObject = InExpandObject;
+	}
+
+	void SetExpandScriptStruct(bool bInExpandScriptStruct)
+	{
+		bExpandScriptStruct = bInExpandScriptStruct;
+	}
+	
+	void SetExpandFunction(bool bInExpandFunction)
+	{
+		bExpandFunction = bInExpandFunction;
+	}
+
+private:
+	EObjectExpandFlag ExpandObject = EObjectExpandFlag::None;
+	bool bExpandScriptStruct = false;
+	bool bExpandFunction = false;
 };
+ENUM_CLASS_FLAGS(FFieldExpander_Default::EObjectExpandFlag);
 
 } //namespace
diff --git a/Engine/Source/Runtime/AdvancedWidgets/Public/Widgets/PropertyViewer/SPropertyViewer.h b/Engine/Source/Runtime/AdvancedWidgets/Public/Widgets/PropertyViewer/SPropertyViewer.h
--- a/Engine/Source/Runtime/AdvancedWidgets/Public/Widgets/PropertyViewer/SPropertyViewer.h
+++ b/Engine/Source/Runtime/AdvancedWidgets/Public/Widgets/PropertyViewer/SPropertyViewer.h
@@ -24,6 +24,9 @@ class FPropertyViewerImpl;
 /** */
 class ADVANCEDWIDGETS_API SPropertyViewer : public SCompoundWidget
 {
+private:
+	using Super = SCompoundWidget;
+
 public:
 	class FHandle
 	{
@@ -115,6 +118,8 @@ public:
 	void Construct(const FArguments& InArgs, UObject* ObjectInstance);
 	void Construct(const FArguments& InArgs, const UFunction* Function);
 
+	virtual void Tick(const FGeometry& AllottedGeometry, const double InCurrentTime, const float InDeltaTime) override;
+
 public:
 	FHandle AddContainer(const UScriptStruct* Struct, TOptional<FText> DisplayName = TOptional<FText>());
 	FHandle AddContainer(const UClass* Class, TOptional<FText> DisplayName = TOptional<FText>());
