Added SPropertyViewer::SetSelection() to set the currently selected property.

#rnx

diff --git a/Engine/Source/Runtime/AdvancedWidgets/Private/Widgets/PropertyViewer/SPropertyViewer.cpp b/Engine/Source/Runtime/AdvancedWidgets/Private/Widgets/PropertyViewer/SPropertyViewer.cpp
--- a/Engine/Source/Runtime/AdvancedWidgets/Private/Widgets/PropertyViewer/SPropertyViewer.cpp
+++ b/Engine/Source/Runtime/AdvancedWidgets/Private/Widgets/PropertyViewer/SPropertyViewer.cpp
@@ -118,6 +118,10 @@ void SPropertyViewer::SetRawFilterText(const FText& InFilterText)
 	Implementation->SetRawFilterText(InFilterText);
 }
 
+void SPropertyViewer::SetSelection(FHandle Container, TArrayView<const FFieldVariant> FieldPath)
+{
+	Implementation->SetSelection(Container, FieldPath);
+}
 
 SPropertyViewer::FHandle SPropertyViewer::MakeContainerIdentifier()
 {
diff --git a/Engine/Source/Runtime/AdvancedWidgets/Private/Widgets/PropertyViewer/SPropertyViewerImpl.cpp b/Engine/Source/Runtime/AdvancedWidgets/Private/Widgets/PropertyViewer/SPropertyViewerImpl.cpp
--- a/Engine/Source/Runtime/AdvancedWidgets/Private/Widgets/PropertyViewer/SPropertyViewerImpl.cpp
+++ b/Engine/Source/Runtime/AdvancedWidgets/Private/Widgets/PropertyViewer/SPropertyViewerImpl.cpp
@@ -742,6 +742,50 @@ FText FPropertyViewerImpl::SetRawFilterTextInternal(const FText& InFilterText)
 	return SearchFilter->GetFilterErrorText();
 }
 
+TSharedPtr<FTreeNode> FPropertyViewerImpl::FindExistingChild(const TSharedPtr<FTreeNode>& Node, TArrayView<const FFieldVariant> FieldPath) const
+{
+	if (FieldPath.Num() == 0)
+	{
+		return TSharedPtr<FTreeNode>();
+	}
+
+	for (const TSharedPtr<FTreeNode>& Child : Node->ChildNodes)
+	{
+		if (Child->GetField() == FieldPath[0])
+		{
+			if (FieldPath.Num() == 1)
+			{
+				return Child;
+			}
+			else
+			{
+				return FindExistingChild(Child, FieldPath.RightChop(1));
+			}
+		}
+	}
+
+	return TSharedPtr<FTreeNode>();
+}
+
+void FPropertyViewerImpl::SetSelection(SPropertyViewer::FHandle Identifier, TArrayView<const FFieldVariant> FieldPath)
+{
+	for (const TSharedPtr<FTreeNode>& Node : TreeSource)
+	{
+		if (TSharedPtr<FContainer> Container = Node->GetContainer())
+		{
+			if (Container->GetIdentifier() == Identifier)
+			{
+				TreeWidget->ClearSelection();
+
+				if (TSharedPtr<FTreeNode> FoundNode = FindChild(Node, FieldPath))
+				{
+					TreeWidget->SetItemSelection(FoundNode, true);
+				}
+			}
+		}
+	}
+}
+
 
 TSharedRef<SWidget> FPropertyViewerImpl::CreateTree(bool bHasPreWidget, bool bShowPropertyValue, bool bHasPostWidget, ESelectionMode::Type SelectionMode)
 {
diff --git a/Engine/Source/Runtime/AdvancedWidgets/Private/Widgets/PropertyViewer/SPropertyViewerImpl.h b/Engine/Source/Runtime/AdvancedWidgets/Private/Widgets/PropertyViewer/SPropertyViewerImpl.h
--- a/Engine/Source/Runtime/AdvancedWidgets/Private/Widgets/PropertyViewer/SPropertyViewerImpl.h
+++ b/Engine/Source/Runtime/AdvancedWidgets/Private/Widgets/PropertyViewer/SPropertyViewerImpl.h
@@ -245,6 +245,7 @@ public:
 	void RemoveAll();
 	TArray<SPropertyViewer::FSelectedItem> GetSelectedItems() const;
 	void SetRawFilterText(const FText& InFilterText);
+	void SetSelection(SPropertyViewer::FHandle Identifier, TArrayView<const FFieldVariant> FieldPath);
 
 private:
 	void AddContainerInternal(SPropertyViewer::FHandle Identifier, TSharedPtr<FContainer>& NewContainer);
@@ -264,6 +265,8 @@ private:
 	void HandleSelectionChanged(TSharedPtr<FTreeNode> Item, ESelectInfo::Type SelectionType);
 	void HandleDoubleClick(TSharedPtr<FTreeNode> Item);
 
+	TSharedPtr<FTreeNode> FindChild(const TSharedPtr<FTreeNode>& ContainerNode, TArrayView<const FFieldVariant> FieldPath) const;
+
 #if WITH_EDITOR
 	void HandleBlueprintCompiled();
 #endif //WITH_EDITOR
diff --git a/Engine/Source/Runtime/AdvancedWidgets/Public/Widgets/PropertyViewer/SPropertyViewer.h b/Engine/Source/Runtime/AdvancedWidgets/Public/Widgets/PropertyViewer/SPropertyViewer.h
--- a/Engine/Source/Runtime/AdvancedWidgets/Public/Widgets/PropertyViewer/SPropertyViewer.h
+++ b/Engine/Source/Runtime/AdvancedWidgets/Public/Widgets/PropertyViewer/SPropertyViewer.h
@@ -133,6 +133,7 @@ public:
 	TArray<FSelectedItem> GetSelectedItems() const;
 
 	void SetRawFilterText(const FText& InFilterText);
+	void SetSelection(FHandle Container, TArrayView<const FFieldVariant> FieldPath);
 
 private:
 	void ConstructInternal(const FArguments& InArgs);
