Missed changes to unsaved files.

#rnx

diff --git a/Engine/Source/Runtime/AdvancedWidgets/Private/Widgets/PropertyViewer/SPropertyViewerImpl.cpp b/Engine/Source/Runtime/AdvancedWidgets/Private/Widgets/PropertyViewer/SPropertyViewerImpl.cpp
--- a/Engine/Source/Runtime/AdvancedWidgets/Private/Widgets/PropertyViewer/SPropertyViewerImpl.cpp
+++ b/Engine/Source/Runtime/AdvancedWidgets/Private/Widgets/PropertyViewer/SPropertyViewerImpl.cpp
@@ -777,7 +777,7 @@ void FPropertyViewerImpl::SetSelection(SPropertyViewer::FHandle Identifier, TArr
 			{
 				TreeWidget->ClearSelection();
 
-				if (TSharedPtr<FTreeNode> FoundNode = FindChild(Node, FieldPath))
+				if (TSharedPtr<FTreeNode> FoundNode = FindExistingChild(Node, FieldPath))
 				{
 					TreeWidget->SetItemSelection(FoundNode, true);
 				}
diff --git a/Engine/Source/Runtime/AdvancedWidgets/Private/Widgets/PropertyViewer/SPropertyViewerImpl.h b/Engine/Source/Runtime/AdvancedWidgets/Private/Widgets/PropertyViewer/SPropertyViewerImpl.h
--- a/Engine/Source/Runtime/AdvancedWidgets/Private/Widgets/PropertyViewer/SPropertyViewerImpl.h
+++ b/Engine/Source/Runtime/AdvancedWidgets/Private/Widgets/PropertyViewer/SPropertyViewerImpl.h
@@ -265,7 +265,7 @@ private:
 	void HandleSelectionChanged(TSharedPtr<FTreeNode> Item, ESelectInfo::Type SelectionType);
 	void HandleDoubleClick(TSharedPtr<FTreeNode> Item);
 
-	TSharedPtr<FTreeNode> FindChild(const TSharedPtr<FTreeNode>& ContainerNode, TArrayView<const FFieldVariant> FieldPath) const;
+	TSharedPtr<FTreeNode> FindExistingChild(const TSharedPtr<FTreeNode>& ContainerNode, TArrayView<const FFieldVariant> FieldPath) const;
 
 #if WITH_EDITOR
 	void HandleBlueprintCompiled();
