MVVM: Add the BindingType icon in front of the binding for the MVVM view.

diff --git a/Engine/Plugins/Runtime/ModelViewViewModel/Source/ModelViewViewModel/Private/MVVMSubsystem.cpp b/Engine/Plugins/Runtime/ModelViewViewModel/Source/ModelViewViewModel/Private/MVVMSubsystem.cpp
--- a/Engine/Plugins/Runtime/ModelViewViewModel/Source/ModelViewViewModel/Private/MVVMSubsystem.cpp
+++ b/Engine/Plugins/Runtime/ModelViewViewModel/Source/ModelViewViewModel/Private/MVVMSubsystem.cpp
@@ -197,19 +197,18 @@ namespace UE::MVVM::Private
 	}
 
 
-	FMVVMAvailableBinding GetAvailableBinding(FMVVMBindingName BindingName, const UStruct* Container, const UClass* AccessorType, const UE::FieldNotification::IClassDescriptor* ClassDescriptor)
+	FMVVMAvailableBinding GetAvailableBinding(FMVVMConstFieldVariant FieldVariant, const UClass* AccessorType, const UE::FieldNotification::IClassDescriptor* ClassDescriptor)
 	{
 		bool bHasNotify = false;
 		if (ClassDescriptor)
 		{
-			UE::FieldNotification::FFieldId FieldId = ClassDescriptor->GetField(CastChecked<UClass>(Container), BindingName.ToName());
+			UE::FieldNotification::FFieldId FieldId = ClassDescriptor->GetField(CastChecked<UClass>(FieldVariant.GetOwner()), FieldVariant.GetName());
 			bHasNotify = FieldId.IsValid();
 		}
 
-		const bool bCanAccessPrivateMember = Container == AccessorType;
-		const bool bCanAccessProtectedMember = AccessorType ? AccessorType->IsChildOf(Container) : false;
+		const bool bCanAccessPrivateMember = FieldVariant.GetOwner() == AccessorType;
+		const bool bCanAccessProtectedMember = AccessorType ? AccessorType->IsChildOf(FieldVariant.GetOwner()) : false;
 
-		FMVVMFieldVariant FieldVariant = UE::MVVM::BindingHelper::FindFieldByName(Container, BindingName);
 		if (FieldVariant.IsProperty())
 		{
 			return GetAvailableBinding(FieldVariant.GetProperty(), bHasNotify, bCanAccessPrivateMember, bCanAccessProtectedMember);
@@ -222,23 +221,31 @@ namespace UE::MVVM::Private
 		return FMVVMAvailableBinding();
 	}
 
-	FMVVMAvailableBinding GetAvailableBinding(FMVVMBindingName BindingName, const UClass* InSubClass, const UClass* InAccessor)
-	{
-		if (!BindingName.IsValid())
-		{
-			return FMVVMAvailableBinding();
-		}
 
-		if (InSubClass && InSubClass->ImplementsInterface(UNotifyFieldValueChanged::StaticClass()))
+	FMVVMAvailableBinding GetAvailableBinding(FMVVMConstFieldVariant FieldVariant, const UClass* InAccessor)
+	{
+		const UClass* OwnerClass = Cast<UClass>(FieldVariant.GetOwner());
+		if (OwnerClass && OwnerClass->ImplementsInterface(UNotifyFieldValueChanged::StaticClass()))
 		{
-			TScriptInterface<INotifyFieldValueChanged> DefaultObject = InSubClass->GetDefaultObject();
+			TScriptInterface<INotifyFieldValueChanged> DefaultObject = OwnerClass->GetDefaultObject();
 			const UE::FieldNotification::IClassDescriptor& ClassDescriptor = DefaultObject->GetFieldNotificationDescriptor();
-			return GetAvailableBinding(BindingName, InSubClass, InAccessor, &ClassDescriptor);
+			return GetAvailableBinding(FieldVariant, InAccessor, &ClassDescriptor);
 		}
 		else
 		{
-			return GetAvailableBinding(BindingName, InSubClass, InAccessor, nullptr);
+			return GetAvailableBinding(FieldVariant, InAccessor, nullptr);
+		}
+	}
+
+	FMVVMAvailableBinding GetAvailableBinding(FMVVMBindingName BindingName, const UClass* InSubClass, const UClass* InAccessor)
+	{
+		if (!BindingName.IsValid())
+		{
+			return FMVVMAvailableBinding();
 		}
+
+		FMVVMConstFieldVariant FieldVariant = UE::MVVM::BindingHelper::FindFieldByName(InSubClass, BindingName);
+		return GetAvailableBinding(FieldVariant, InAccessor);
 	}
 } //namespace
 
@@ -261,6 +268,12 @@ FMVVMAvailableBinding UMVVMSubsystem::GetAvailableBinding(const UClass* Class, F
 }
 
 
+FMVVMAvailableBinding UMVVMSubsystem::GetAvailableBindingForVariant(UE::MVVM::FMVVMConstFieldVariant FieldVariant, const UClass* Accessor) const
+{
+	return UE::MVVM::Private::GetAvailableBinding(FieldVariant, Accessor);
+}
+
+
 TValueOrError<bool, FText> UMVVMSubsystem::IsBindingValid(FBindingArgs Args) const
 {
 	if (UE::MVVM::IsForwardBinding(Args.Mode))
diff --git a/Engine/Plugins/Runtime/ModelViewViewModel/Source/ModelViewViewModel/Public/MVVMSubsystem.h b/Engine/Plugins/Runtime/ModelViewViewModel/Source/ModelViewViewModel/Public/MVVMSubsystem.h
--- a/Engine/Plugins/Runtime/ModelViewViewModel/Source/ModelViewViewModel/Public/MVVMSubsystem.h
+++ b/Engine/Plugins/Runtime/ModelViewViewModel/Source/ModelViewViewModel/Public/MVVMSubsystem.h
@@ -35,24 +35,23 @@ public:
 	UFUNCTION(BlueprintCallable, Category = "MVVM")
 	bool DoesWidgetTreeContainedWidget(const UWidgetTree* WidgetTree, const UWidget* ViewWidget) const;
 
-	/**
-	 * @return The list of all the bindings that are available for the Class.
-	 */
+	/** @return The list of all the AvailableBindings that are available for the Class. */
 	UFUNCTION(BlueprintCallable, Category = "MVVM")
 	TArray<FMVVMAvailableBinding> GetAvailableBindings(const UClass* Class, const UClass* Accessor) const;
 
 	/**
-	 * @return The list of all the bindings that are available from the SriptStuct.
+	 * @return The list of all the AvailableBindings that are available from the SriptStuct.
 	 * @note When FMVVMAvailableBinding::HasNotify is false, a notification can still be triggered by the owner of the struct. The struct changed but which property of the struct changed is unknown.
 	 */
 	TArray<FMVVMAvailableBinding> GetAvailableBindingsForStruct(const UScriptStruct* Struct) const;
 
-	/**
-	 * @return the available binding from the binding.
-	 */
+	/** @return The AvailableBinding from a BindingName. */
 	UFUNCTION(BlueprintCallable, Category = "MVVM")
 	FMVVMAvailableBinding GetAvailableBinding(const UClass* Class, FMVVMBindingName BindingName, const UClass* Accessor) const;
 
+	/** @return The AvailableBinding from a field. */
+	FMVVMAvailableBinding GetAvailableBindingForVariant(UE::MVVM::FMVVMConstFieldVariant Variant, const UClass* Accessor) const;
+
 	UFUNCTION(BlueprintCallable, Category = "MVVM")
 	UMVVMViewModelCollectionObject* GetGlobalViewModelCollection() const
 	{
diff --git a/Engine/Plugins/Runtime/ModelViewViewModel/Source/ModelViewViewModelEditor/Private/Widgets/SMVVMSelectViewModel.cpp b/Engine/Plugins/Runtime/ModelViewViewModel/Source/ModelViewViewModelEditor/Private/Widgets/SMVVMSelectViewModel.cpp
--- a/Engine/Plugins/Runtime/ModelViewViewModel/Source/ModelViewViewModelEditor/Private/Widgets/SMVVMSelectViewModel.cpp
+++ b/Engine/Plugins/Runtime/ModelViewViewModel/Source/ModelViewViewModelEditor/Private/Widgets/SMVVMSelectViewModel.cpp
@@ -99,6 +99,7 @@ void SMVVMSelectViewModel::Construct(const FArguments& InArgs, const UWidgetBlue
 						[
 							SAssignNew(BindingListWidget, SSourceBindingList, WidgetBlueprint)
 							.EnableSelection(false)
+							.FieldVisibilityFlags(EFieldVisibility::None)
 						]
 					]
 				]
diff --git a/Engine/Plugins/Runtime/ModelViewViewModel/Source/ModelViewViewModelEditor/Private/Widgets/SMVVMViewModelBindingListWidget.cpp b/Engine/Plugins/Runtime/ModelViewViewModel/Source/ModelViewViewModelEditor/Private/Widgets/SMVVMViewModelBindingListWidget.cpp
--- a/Engine/Plugins/Runtime/ModelViewViewModel/Source/ModelViewViewModelEditor/Private/Widgets/SMVVMViewModelBindingListWidget.cpp
+++ b/Engine/Plugins/Runtime/ModelViewViewModel/Source/ModelViewViewModelEditor/Private/Widgets/SMVVMViewModelBindingListWidget.cpp
@@ -16,10 +16,15 @@
 #include "MVVMEditorSubsystem.h"
 #include "MVVMSubsystem.h"
 #include "MVVMViewModelBase.h"
+#include "Styling/MVVMEditorStyle.h"
 #include "Types/MVVMAvailableBinding.h"
 #include "Types/MVVMFieldVariant.h"
 #include "WidgetBlueprint.h"
+#include "Widgets/Images/SImage.h"
+#include "Widgets/Images/SLayeredImage.h"
+#include "Widgets/Layout/SBox.h"
 #include "Widgets/PropertyViewer/SPropertyViewer.h"
+#include "Widgets/SNullWidget.h"
 
 #define LOCTEXT_NAMESPACE "SSourceBindingList"
 
@@ -250,6 +255,56 @@ TOptional<const UStruct*> FFieldExpander_Bindable::GetExpandedFunction(const UFu
 	return TOptional<const UStruct*>();
 }
 
+TSharedRef<SWidget> ConstructFieldPreSlot(const UWidgetBlueprint* WidgetBlueprint, UE::PropertyViewer::SPropertyViewer::FHandle Handle, const FFieldVariant FieldPath)
+{
+	TSharedRef<SWidget> ImageWidget = SNullWidget::NullWidget;
+	TSubclassOf<UObject> AccessorClass = WidgetBlueprint ? WidgetBlueprint->SkeletonGeneratedClass : nullptr;
+	FMVVMAvailableBinding Binding = GEngine->GetEngineSubsystem<UMVVMSubsystem>()->GetAvailableBindingForVariant(FMVVMConstFieldVariant(FieldPath), AccessorClass);
+	if (Binding.IsValid())
+	{
+		const FSlateBrush* Brush = nullptr;
+		const FLinearColor ImageColor = FLinearColor(1.0f, 1.0f, 1.0, 0.5f);
+		if (Binding.HasNotify() && Binding.IsReadable() && Binding.IsWritable())
+		{
+			Brush = FMVVMEditorStyle::Get().GetBrush("BindingMode.TwoWay");
+		}
+		else if (Binding.HasNotify() && Binding.IsReadable())
+		{
+			Brush = FMVVMEditorStyle::Get().GetBrush("BindingMode.OneWayToSource");
+		}
+		else if (Binding.IsWritable() && Binding.IsReadable())
+		{
+			ImageWidget = SNew(SLayeredImage, FMVVMEditorStyle::Get().GetBrush("BindingMode.OneWay"), ImageColor)
+				.Image(FMVVMEditorStyle::Get().GetBrush("BindingMode.OneTime"))
+				.ColorAndOpacity(ImageColor);
+		}
+		else if (Binding.IsWritable())
+		{
+			Brush = FMVVMEditorStyle::Get().GetBrush("BindingMode.OneWay");
+		}
+		else if (Binding.IsReadable())
+		{
+			Brush = FMVVMEditorStyle::Get().GetBrush("BindingMode.OneTime");
+		}
+
+		if (Brush)
+		{
+			ImageWidget = SNew(SImage)
+				.Image(Brush)
+				.ColorAndOpacity(ImageColor);
+		}
+	}
+
+	return SNew(SBox)
+		.HAlign(HAlign_Center)
+		.VAlign(VAlign_Center)
+		.WidthOverride(16.0f)
+		.HeightOverride(16.0f)
+		[
+			ImageWidget
+		];
+}
+
 /** */
 void SSourceBindingList::Construct(const FArguments& InArgs, const UWidgetBlueprint* InWidgetBlueprint)
 {
@@ -268,6 +323,7 @@ void SSourceBindingList::Construct(const FArguments& InArgs, const UWidgetBluepr
 		.bSanitizeName(true)
 		.SelectionMode(InArgs._EnableSelection ? ESelectionMode::Single : ESelectionMode::None)
 		.bShowSearchBox(InArgs._ShowSearchBox)
+		.OnGetPreSlot(this, &SSourceBindingList::HandleGetPreSlot)
 		.OnSelectionChanged(this, &SSourceBindingList::HandleSelectionChanged)
 		.OnDoubleClicked(this, &SSourceBindingList::HandleDoubleClicked);
 
@@ -479,6 +535,15 @@ FMVVMBlueprintPropertyPath SSourceBindingList::CreateBlueprintPropertyPath(SProp
 	return PropertyPath;
 }
 
+TSharedPtr<SWidget> SSourceBindingList::HandleGetPreSlot(SPropertyViewer::FHandle Handle, TArrayView<const FFieldVariant> FieldPath)
+{
+	if (FieldPath.Num() > 0)
+	{
+		return ConstructFieldPreSlot(WidgetBlueprint.Get(), Handle, FieldPath.Last());
+	}
+	return TSharedPtr<SWidget>();
+}
+
 void SSourceBindingList::HandleSelectionChanged(SPropertyViewer::FHandle Handle, TArrayView<const FFieldVariant> FieldPath, ESelectInfo::Type SelectionType)
 {
 	SelectedPath = CreateBlueprintPropertyPath(Handle, FieldPath);
diff --git a/Engine/Plugins/Runtime/ModelViewViewModel/Source/ModelViewViewModelEditor/Private/Widgets/SMVVMViewModelBindingListWidget.h b/Engine/Plugins/Runtime/ModelViewViewModel/Source/ModelViewViewModelEditor/Private/Widgets/SMVVMViewModelBindingListWidget.h
--- a/Engine/Plugins/Runtime/ModelViewViewModel/Source/ModelViewViewModelEditor/Private/Widgets/SMVVMViewModelBindingListWidget.h
+++ b/Engine/Plugins/Runtime/ModelViewViewModel/Source/ModelViewViewModelEditor/Private/Widgets/SMVVMViewModelBindingListWidget.h
@@ -56,8 +56,8 @@ namespace UE::MVVM
 	};
 
 	/**
- *
- */
+	 *
+	 */
 	class FFieldExpander_Bindable : public UE::PropertyViewer::FFieldExpander_Default
 	{
 	public:
@@ -65,6 +65,9 @@ namespace UE::MVVM
 		virtual TOptional<const UStruct*> GetExpandedFunction(const UFunction* Function) const override;
 	};
 
+	/** Genereate the icon for the binding type. */
+	TSharedRef<SWidget> ConstructFieldPreSlot(const UWidgetBlueprint* WidgetBlueprint, UE::PropertyViewer::SPropertyViewer::FHandle Handle, const FFieldVariant FieldPath);
+
 	/** 
 	 * 
 	 */
@@ -108,8 +111,9 @@ namespace UE::MVVM
 	private:
 		using SPropertyViewer = UE::PropertyViewer::SPropertyViewer;
 
-		void HandleSelectionChanged(SPropertyViewer::FHandle ViewModel, TArrayView<const FFieldVariant> Path, ESelectInfo::Type SelectionType);
-		void HandleDoubleClicked(SPropertyViewer::FHandle ViewModel, TArrayView<const FFieldVariant> Path);
+		TSharedPtr<SWidget> HandleGetPreSlot(SPropertyViewer::FHandle Handle, TArrayView<const FFieldVariant> Path);
+		void HandleSelectionChanged(SPropertyViewer::FHandle Handle, TArrayView<const FFieldVariant> Path, ESelectInfo::Type SelectionType);
+		void HandleDoubleClicked(SPropertyViewer::FHandle Handle, TArrayView<const FFieldVariant> Path);
 
 		FMVVMBlueprintPropertyPath CreateBlueprintPropertyPath(SPropertyViewer::FHandle Handle, TArrayView<const FFieldVariant> Path) const;
 
diff --git a/Engine/Plugins/Runtime/ModelViewViewModel/Source/ModelViewViewModelEditor/Private/Widgets/SMVVMViewModelPanel.cpp b/Engine/Plugins/Runtime/ModelViewViewModel/Source/ModelViewViewModelEditor/Private/Widgets/SMVVMViewModelPanel.cpp
--- a/Engine/Plugins/Runtime/ModelViewViewModel/Source/ModelViewViewModelEditor/Private/Widgets/SMVVMViewModelPanel.cpp
+++ b/Engine/Plugins/Runtime/ModelViewViewModel/Source/ModelViewViewModelEditor/Private/Widgets/SMVVMViewModelPanel.cpp
@@ -50,7 +50,8 @@ void SMVVMViewModelPanel::Construct(const FArguments& InArgs, TSharedPtr<FWidget
 
 	WeakBlueprintEditor = WidgetBlueprintEditor;
 	WeakBlueprintView = CurrentBlueprintView;
-	FieldIterator = MakeUnique<FFieldIterator_Bindable>(WidgetBlueprint, EFieldVisibility::Notify);
+	FieldIterator = MakeUnique<FFieldIterator_Bindable>(WidgetBlueprint, EFieldVisibility::None);
+	FieldExpander = MakeUnique<FFieldExpander_Bindable>();
 
 	if (CurrentBlueprintView)
 	{
@@ -66,6 +67,8 @@ void SMVVMViewModelPanel::Construct(const FArguments& InArgs, TSharedPtr<FWidget
 		.bShowFieldIcon(true)
 		.bSanitizeName(true)
 		.FieldIterator(FieldIterator.Get())
+		.FieldExpander(FieldExpander.Get())
+		.OnGetPreSlot(this, &SMVVMViewModelPanel::HandleGetPreSlot)
 		.OnContextMenuOpening(this, &SMVVMViewModelPanel::HandleContextMenuOpening)
 		.OnSelectionChanged(this, &SMVVMViewModelPanel::HandleSelectionChanged)
 		.OnGenerateContainer(this, &SMVVMViewModelPanel::HandleGenerateContainer)
@@ -275,6 +278,23 @@ bool SMVVMViewModelPanel::HandleCanEditViewmodelList() const
 }
 
 
+TSharedPtr<SWidget> SMVVMViewModelPanel::HandleGetPreSlot(UE::PropertyViewer::SPropertyViewer::FHandle Handle, TArrayView<const FFieldVariant> FieldPath)
+{
+	if (FieldPath.Num() > 0)
+	{
+		UWidgetBlueprint* WidgetBlueprint = nullptr;
+		if (TSharedPtr<FWidgetBlueprintEditor> WidgetBlueprintEditor = WeakBlueprintEditor.Pin())
+		{
+			WidgetBlueprint = WidgetBlueprintEditor->GetWidgetBlueprintObj();
+		}
+
+		return ConstructFieldPreSlot(WidgetBlueprint, Handle, FieldPath.Last());
+	}
+
+	return TSharedPtr<SWidget>();
+}
+
+
 TSharedRef<SWidget> SMVVMViewModelPanel::HandleGenerateContainer(UE::PropertyViewer::SPropertyViewer::FHandle ContainerHandle, TOptional<FText> DisplayName)
 {
 	if (const FGuid* VMGuidPtr = PropertyViewerHandles.Find(ContainerHandle))
@@ -468,6 +488,7 @@ void SMVVMViewModelPanel::HandleSelectionChanged(UE::PropertyViewer::SPropertyVi
 		if (const FGuid* VMGuidPtr = PropertyViewerHandles.Find(ContainerHandle))
 		{
 			SelectedViewModelGuid = *VMGuidPtr;
+
 			if (UMVVMBlueprintView* BlueprintView = WeakBlueprintView.Get())
 			{
 				if (FMVVMBlueprintViewModelContext* ViewModelContext = BlueprintView->FindViewModel(SelectedViewModelGuid))
diff --git a/Engine/Plugins/Runtime/ModelViewViewModel/Source/ModelViewViewModelEditor/Private/Widgets/SMVVMViewModelPanel.h b/Engine/Plugins/Runtime/ModelViewViewModel/Source/ModelViewViewModelEditor/Private/Widgets/SMVVMViewModelPanel.h
--- a/Engine/Plugins/Runtime/ModelViewViewModel/Source/ModelViewViewModelEditor/Private/Widgets/SMVVMViewModelPanel.h
+++ b/Engine/Plugins/Runtime/ModelViewViewModel/Source/ModelViewViewModelEditor/Private/Widgets/SMVVMViewModelPanel.h
@@ -40,6 +40,7 @@ private:
 	TSharedRef<SWidget> MakeAddMenu();
 	void HandleCancelAddMenu();
 	void HandleAddMenuViewModel(const UClass* SelectedClass);
+	TSharedPtr<SWidget> HandleGetPreSlot(UE::PropertyViewer::SPropertyViewer::FHandle Handle, TArrayView<const FFieldVariant> FieldPath);
 	bool HandleCanEditViewmodelList() const;
 	FReply OnKeyDown(const FGeometry& MyGeometry, const FKeyEvent& InKeyEvent) override;
 	TSharedRef<SWidget> HandleGenerateContainer(UE::PropertyViewer::SPropertyViewer::FHandle ContainerHandle, TOptional<FText> DisplayName);
@@ -66,6 +67,7 @@ private:
 	TSharedPtr<SPositiveActionButton> AddMenuButton;
 	TSharedPtr<UE::PropertyViewer::SPropertyViewer> ViewModelTreeView;
 	TUniquePtr<FFieldIterator_Bindable> FieldIterator;
+	TUniquePtr<UE::PropertyViewer::FFieldExpander_Default> FieldExpander;
 	TSharedPtr<FUICommandList> CommandList;
 	TSharedPtr<IStructureDetailsView> PropertyView;
 
diff --git a/Engine/Source/Runtime/AdvancedWidgets/Private/Widgets/PropertyViewer/SPropertyViewerImpl.cpp b/Engine/Source/Runtime/AdvancedWidgets/Private/Widgets/PropertyViewer/SPropertyViewerImpl.cpp
--- a/Engine/Source/Runtime/AdvancedWidgets/Private/Widgets/PropertyViewer/SPropertyViewerImpl.cpp
+++ b/Engine/Source/Runtime/AdvancedWidgets/Private/Widgets/PropertyViewer/SPropertyViewerImpl.cpp
@@ -39,7 +39,6 @@ namespace UE::PropertyViewer::Private
 /** 
  * Column name 
  */
-static FName ColumnName_FieldPreWidget = "FieldPreWidget";
 static FName ColumnName_Field = "Field";
 static FName ColumnName_PropertyValue = "FieldValue";
 static FName ColumnName_FieldPostWidget = "FieldPostWidget";
@@ -795,22 +794,13 @@ void FPropertyViewerImpl::SetSelection(SPropertyViewer::FHandle Identifier, TArr
 TSharedRef<SWidget> FPropertyViewerImpl::CreateTree(bool bHasPreWidget, bool bShowPropertyValue, bool bHasPostWidget, ESelectionMode::Type SelectionMode)
 {
 	TSharedPtr<SHeaderRow> HeaderRowWidget;
-	if (bHasPreWidget || bShowPropertyValue || bHasPostWidget)
+	if (bShowPropertyValue || bHasPostWidget)
 	{
 		bUseRows = true;
 
 		HeaderRowWidget = SNew(SHeaderRow)
 		.Visibility(EVisibility::Collapsed);
 
-		if (bHasPreWidget)
-		{
-			HeaderRowWidget->AddColumn(
-				SHeaderRow::FColumn::FArguments()
-				.ColumnId(ColumnName_FieldPreWidget)
-				.DefaultLabel(LOCTEXT("PropertyPreWidget", ""))
-			);
-		}
-
 		HeaderRowWidget->AddColumn(
 			SHeaderRow::FColumn::FArguments()
 			.ColumnId(ColumnName_Field)
@@ -897,6 +887,13 @@ TSharedRef<ITableRow> FPropertyViewerImpl::HandleGenerateRow(TSharedPtr<FTreeNod
 {
 	FText HighlightText = SearchFilter ? SearchFilter->GetRawFilterText() : FText::GetEmpty();
 
+	TSharedPtr<SWidget> PreWidget;
+	if (OnGetPreSlot.IsBound())
+	{
+		TSharedPtr<FContainer> OwnerContainer = Item->GetOwnerContainer();
+		PreWidget = OnGetPreSlot.Execute(OwnerContainer.IsValid() ? OwnerContainer->GetIdentifier() : SPropertyViewer::FHandle(), Item->GetFieldPath());
+	}
+
 	TSharedPtr<SWidget> ItemWidget;
 	if (TSharedPtr<FContainer> ContainerPin = Item->GetContainer())
 	{
@@ -966,21 +963,45 @@ TSharedRef<ITableRow> FPropertyViewerImpl::HandleGenerateRow(TSharedPtr<FTreeNod
 		{
 			if (ColumnName == ColumnName_Field)
 			{
-				return SNew(SHorizontalBox)
-				+ SHorizontalBox::Slot()
+				TSharedRef<SHorizontalBox> Stack = SNew(SHorizontalBox);
+
+				TSharedPtr<FPropertyViewerImpl> PropertyViewOwnerPin = PropertyViewOwner.Pin();
+				TSharedPtr<SWidget> PreWidget;
+				if (PropertyViewOwnerPin && PropertyViewOwnerPin->OnGetPreSlot.IsBound())
+				{
+					if (TSharedPtr<FTreeNode> ItemPin = Item.Pin())
+					{						
+						TSharedPtr<FContainer> OwnerContainer = ItemPin->GetOwnerContainer();
+						PreWidget = PropertyViewOwnerPin->OnGetPreSlot.Execute(OwnerContainer.IsValid() ? OwnerContainer->GetIdentifier() : SPropertyViewer::FHandle(), ItemPin->GetFieldPath());
+					}
+				}
+
+				Stack->AddSlot()
 				.AutoWidth()
 				[
 					SNew(SExpanderArrow, SharedThis(this))
 					.IndentAmount(16)
 					.ShouldDrawWires(true)
-				]
-				+SHorizontalBox::Slot()
+				];
+
+				if (PreWidget)
+				{
+					Stack->AddSlot()
+					.AutoWidth()
+					[
+						PreWidget.ToSharedRef()
+					];
+				}
+
+				Stack->AddSlot()
 				.AutoWidth()
 				.Padding(2.0f, 0.0f)
 				.VAlign(VAlign_Center)
 				[
 					FieldWidget.ToSharedRef()
 				];
+
+				return Stack;
 			}
 
 			if (ColumnName == ColumnName_PropertyValue)
@@ -1035,21 +1056,17 @@ TSharedRef<ITableRow> FPropertyViewerImpl::HandleGenerateRow(TSharedPtr<FTreeNod
 				}
 			}
 
-			if (ColumnName == ColumnName_FieldPreWidget || ColumnName == ColumnName_FieldPostWidget)
+			if (ColumnName == ColumnName_FieldPostWidget)
 			{
 				TSharedPtr<FTreeNode> ItemPin = Item.Pin();
 				TSharedPtr<FPropertyViewerImpl> PropertyViewOwnerPin = PropertyViewOwner.Pin();
 				if (ItemPin && PropertyViewOwnerPin)
 				{
-					SPropertyViewer::FGetFieldWidget& OnGetWidget = (ColumnName == ColumnName_FieldPreWidget) ? PropertyViewOwnerPin->OnGetPreSlot : PropertyViewOwnerPin->OnGetPostSlot;
-					TSharedPtr<SWidget> PreWidget;
-
 					TSharedPtr<FContainer> OwnerContainer = ItemPin->GetOwnerContainer();
-					PreWidget = OnGetWidget.Execute(OwnerContainer.IsValid() ? OwnerContainer->GetIdentifier() : SPropertyViewer::FHandle(), ItemPin->GetFieldPath());
-
-					if (PreWidget)
+					TSharedPtr<SWidget> PostWidget = PropertyViewOwnerPin->OnGetPostSlot.Execute(OwnerContainer.IsValid() ? OwnerContainer->GetIdentifier() : SPropertyViewer::FHandle(), ItemPin->GetFieldPath());
+					if (PostWidget)
 					{
-						return PreWidget.ToSharedRef();
+						return PostWidget.ToSharedRef();
 					}
 				}
 			}
@@ -1072,6 +1089,27 @@ TSharedRef<ITableRow> FPropertyViewerImpl::HandleGenerateRow(TSharedPtr<FTreeNod
 	}
 
 	using SSimpleRowType = STableRow<TSharedPtr<FTreeNode>>;
+
+	if (PreWidget)
+	{
+		return SNew(SSimpleRowType, OwnerTable)
+			.Padding(0.0f)
+			.Content()
+			[
+				SNew(SHorizontalBox)
+				+ SHorizontalBox::Slot()
+				.AutoWidth()
+				[
+					PreWidget.ToSharedRef()
+				]
+				+ SHorizontalBox::Slot()
+				.FillWidth(1.f)
+				[
+					FieldWidget
+				]
+			];
+	}
+
 	return SNew(SSimpleRowType, OwnerTable)
 		.Padding(0.0f)
 		.Content()
