Slate: Rebuild the PropertyViewer's container when the blueprint is recompiled. The BP is not always mark as "newer class exist" depending if it was a full rebuild or the skeletal rebuild.

diff --git a/Engine/Source/Runtime/AdvancedWidgets/Private/Widgets/PropertyViewer/SPropertyViewerImpl.cpp b/Engine/Source/Runtime/AdvancedWidgets/Private/Widgets/PropertyViewer/SPropertyViewerImpl.cpp
--- a/Engine/Source/Runtime/AdvancedWidgets/Private/Widgets/PropertyViewer/SPropertyViewerImpl.cpp
+++ b/Engine/Source/Runtime/AdvancedWidgets/Private/Widgets/PropertyViewer/SPropertyViewerImpl.cpp
@@ -422,6 +422,7 @@ FPropertyViewerImpl::FPropertyViewerImpl(const SPropertyViewer::FArguments& InAr
 	if (GEditor)
 	{
 		GEditor->OnBlueprintCompiled().AddRaw(this, &FPropertyViewerImpl::HandleBlueprintCompiled);
+		FCoreUObjectDelegates::OnObjectsReplaced.AddRaw(this, &FPropertyViewerImpl::HandleReplaceViewedObjects);
 	}
 #endif
 }
@@ -441,6 +442,7 @@ FPropertyViewerImpl::~FPropertyViewerImpl()
 	if (GEditor)
 	{
 		GEditor->OnBlueprintCompiled().RemoveAll(this);
+		FCoreUObjectDelegates::OnObjectsReplaced.RemoveAll(this);
 	}
 #endif
 }
@@ -1216,15 +1218,38 @@ void FPropertyViewerImpl::HandleBlueprintCompiled()
 
 	if (bRemoved)
 	{
+		TreeWidget->RebuildList();
 		if (FilterHandler)
 		{
 			FilterHandler->RefreshAndFilterTree();
 		}
-		else
+	}
+}
+
+void FPropertyViewerImpl::HandleReplaceViewedObjects(const TMap<UObject*, UObject*>& OldToNewObjectMap)
+{
+	TArray<UObject*> ValueArray;
+	OldToNewObjectMap.GenerateValueArray(ValueArray);
+
+	TArray<TSharedPtr<FContainer>> ItemsToReplace;
+
+	for (TSharedPtr<FContainer> Container : Containers)
+	{
+		if (Container->GetStruct())
 		{
-			TreeWidget->RequestTreeRefresh();
+			if (ValueArray.Contains(Container->GetStruct()))
+			{
+				ItemsToReplace.Add(Container);
+			}
 		}
 	}
+
+	for (TSharedPtr<FContainer> Container : ItemsToReplace)
+	{
+		Remove(Container->GetIdentifier());
+		Containers.Add(Container);
+		AddContainerInternal(Container->GetIdentifier(), Container);
+	}
 }
 #endif //WITH_EDITOR
 
diff --git a/Engine/Source/Runtime/AdvancedWidgets/Private/Widgets/PropertyViewer/SPropertyViewerImpl.h b/Engine/Source/Runtime/AdvancedWidgets/Private/Widgets/PropertyViewer/SPropertyViewerImpl.h
--- a/Engine/Source/Runtime/AdvancedWidgets/Private/Widgets/PropertyViewer/SPropertyViewerImpl.h
+++ b/Engine/Source/Runtime/AdvancedWidgets/Private/Widgets/PropertyViewer/SPropertyViewerImpl.h
@@ -269,6 +269,7 @@ private:
 
 #if WITH_EDITOR
 	void HandleBlueprintCompiled();
+	void HandleReplaceViewedObjects(const TMap<UObject*, UObject*>& OldToNewObjectMap);
 #endif //WITH_EDITOR
 };
 
