PropertyViewer: Fix bad of the map. The Key is the previous object about to be removed and the Value is the new object replacing it.

diff --git a/Engine/Source/Runtime/AdvancedWidgets/Private/Widgets/PropertyViewer/SPropertyViewerImpl.cpp b/Engine/Source/Runtime/AdvancedWidgets/Private/Widgets/PropertyViewer/SPropertyViewerImpl.cpp
--- a/Engine/Source/Runtime/AdvancedWidgets/Private/Widgets/PropertyViewer/SPropertyViewerImpl.cpp
+++ b/Engine/Source/Runtime/AdvancedWidgets/Private/Widgets/PropertyViewer/SPropertyViewerImpl.cpp
@@ -1247,8 +1247,8 @@ void FPropertyViewerImpl::HandleBlueprintCompiled()
 
 void FPropertyViewerImpl::HandleReplaceViewedObjects(const TMap<UObject*, UObject*>& OldToNewObjectMap)
 {
-	TArray<UObject*> ValueArray;
-	OldToNewObjectMap.GenerateValueArray(ValueArray);
+	TArray<UObject*> KeyArray;
+	OldToNewObjectMap.GenerateKeyArray(KeyArray);
 
 	TArray<TSharedPtr<FContainer>> ItemsToReplace;
 
@@ -1256,7 +1256,7 @@ void FPropertyViewerImpl::HandleReplaceViewedObjects(const TMap<UObject*, UObjec
 	{
 		if (Container->GetStruct())
 		{
-			if (ValueArray.Contains(Container->GetStruct()))
+			if (KeyArray.Contains(Container->GetStruct()))
 			{
 				ItemsToReplace.Add(Container);
 			}
