Fix property viewer not displaying correctly when setting a UScriptStruct as a container

Flags for trashed structs were being checked incorrectly

diff --git a/Engine/Source/Runtime/AdvancedWidgets/Private/Widgets/PropertyViewer/SPropertyViewerImpl.cpp b/Engine/Source/Runtime/AdvancedWidgets/Private/Widgets/PropertyViewer/SPropertyViewerImpl.cpp
--- a/Engine/Source/Runtime/AdvancedWidgets/Private/Widgets/PropertyViewer/SPropertyViewerImpl.cpp
+++ b/Engine/Source/Runtime/AdvancedWidgets/Private/Widgets/PropertyViewer/SPropertyViewerImpl.cpp
@@ -83,7 +83,7 @@ bool FContainer::IsValid() const
 	}
 	else if (const UScriptStruct* ScriptStruct = Cast<UScriptStruct>(Container.Get()))
 	{
-		return (ScriptStruct->StructFlags & (EStructFlags::STRUCT_Trashed)) != 0;
+		return (ScriptStruct->StructFlags & (EStructFlags::STRUCT_Trashed)) == 0;
 	}
 	return false;
 }
