[AdvancedWidgets] URadialSlider: changed category name from "Behaviour" to "Behavior" to comply with coding standards (When using Naming Conventions, all code and comments should use U.S. English spelling and grammar.)

diff --git a/Engine/Source/Runtime/AdvancedWidgets/Public/Components/RadialSlider.h b/Engine/Source/Runtime/AdvancedWidgets/Public/Components/RadialSlider.h
--- a/Engine/Source/Runtime/AdvancedWidgets/Public/Components/RadialSlider.h
+++ b/Engine/Source/Runtime/AdvancedWidgets/Public/Components/RadialSlider.h
@@ -162,11 +162,11 @@ public:
 	ADVANCEDWIDGETS_API void SetCustomDefaultValue(float InValue);
 
 	/** Sets the curve for the slider range*/
-	UFUNCTION(BlueprintCallable, Category = "Behaviour")
+	UFUNCTION(BlueprintCallable, Category = "Behavior")
 	ADVANCEDWIDGETS_API void SetSliderRange(const FRuntimeFloatCurve& InSliderRange);
 
 	/** Adds value tags to the slider. */
-	UFUNCTION(BlueprintCallable, Category = "Behaviour")
+	UFUNCTION(BlueprintCallable, Category = "Behavior")
 	ADVANCEDWIDGETS_API void SetValueTags(const TArray<float>& InValueTags);
 
 	/** Sets the minimum angle of the slider. */
@@ -178,7 +178,7 @@ public:
 	ADVANCEDWIDGETS_API void SetSliderHandleEndAngle(float InValue);
 
 	/** Sets the Angular Offset for the slider. */
-	UFUNCTION(BlueprintCallable, Category = "Behaviour")
+	UFUNCTION(BlueprintCallable, Category = "Behavior")
 	ADVANCEDWIDGETS_API void SetAngularOffset(float InValue);
 
 	/** Sets the start and end of the hand as a ratio to the slider radius (so 0.0 to 1.0 is from the slider center to the handle). */
