PropertyViewer: Test if the container is valid before accessing the properties. Fix crash when the BP is compiled while showing the value.

diff --git a/Engine/Source/Runtime/AdvancedWidgets/Private/Widgets/PropertyViewer/SBoolPropertyValue.cpp b/Engine/Source/Runtime/AdvancedWidgets/Private/Widgets/PropertyViewer/SBoolPropertyValue.cpp
--- a/Engine/Source/Runtime/AdvancedWidgets/Private/Widgets/PropertyViewer/SBoolPropertyValue.cpp
+++ b/Engine/Source/Runtime/AdvancedWidgets/Private/Widgets/PropertyViewer/SBoolPropertyValue.cpp
@@ -42,9 +42,9 @@ void SBoolPropertyValue::Construct(const FArguments& InArgs)
 
 ECheckBoxState SBoolPropertyValue::HandleIsChecked() const
 {
-	if (const FBoolProperty* Property = CastField<const FBoolProperty>(Path.GetLastProperty()))
+	if (const void* Container = Path.GetContainerPtr())
 	{
-		if (const void* Container = Path.GetContainerPtr())
+		if (const FBoolProperty* Property = CastField<const FBoolProperty>(Path.GetLastProperty()))
 		{
 			const bool bValue = Property->GetPropertyValue_InContainer(Container);
 			return bValue ? ECheckBoxState::Checked : ECheckBoxState::Unchecked;
@@ -56,9 +56,9 @@ ECheckBoxState SBoolPropertyValue::HandleIsChecked() const
 
 void SBoolPropertyValue::HandleCheckStateChanged(ECheckBoxState NewState)
 {
-	if (const FBoolProperty* Property = CastField<const FBoolProperty>(Path.GetLastProperty()))
+	if (void* Container = Path.GetContainerPtr())
 	{
-		if (void* Container = Path.GetContainerPtr())
+		if (const FBoolProperty* Property = CastField<const FBoolProperty>(Path.GetLastProperty()))
 		{
 			if (NotifyHook)
 			{
diff --git a/Engine/Source/Runtime/AdvancedWidgets/Private/Widgets/PropertyViewer/SDefaultPropertyValue.cpp b/Engine/Source/Runtime/AdvancedWidgets/Private/Widgets/PropertyViewer/SDefaultPropertyValue.cpp
--- a/Engine/Source/Runtime/AdvancedWidgets/Private/Widgets/PropertyViewer/SDefaultPropertyValue.cpp
+++ b/Engine/Source/Runtime/AdvancedWidgets/Private/Widgets/PropertyViewer/SDefaultPropertyValue.cpp
@@ -35,9 +35,9 @@ void SDefaultPropertyValue::Construct(const FArguments& InArgs)
 
 FText SDefaultPropertyValue::GetText() const
 {
-	if (const FProperty* LastProperty = Path.GetLastProperty())
+	if (const void* ContainerPtr = Path.GetContainerPtr())
 	{
-		if (const void* ContainerPtr = Path.GetContainerPtr())
+		if (const FProperty* LastProperty = Path.GetLastProperty())
 		{
 			FString TextVersion;
 			LastProperty->ExportTextItem_InContainer(TextVersion, ContainerPtr, nullptr, nullptr, PPF_SimpleObjectText);
diff --git a/Engine/Source/Runtime/AdvancedWidgets/Private/Widgets/PropertyViewer/SEnumPropertyValue.cpp b/Engine/Source/Runtime/AdvancedWidgets/Private/Widgets/PropertyViewer/SEnumPropertyValue.cpp
--- a/Engine/Source/Runtime/AdvancedWidgets/Private/Widgets/PropertyViewer/SEnumPropertyValue.cpp
+++ b/Engine/Source/Runtime/AdvancedWidgets/Private/Widgets/PropertyViewer/SEnumPropertyValue.cpp
@@ -71,23 +71,17 @@ int32 SEnumPropertyValue::GetCurrentValue() const
 {
 	if (const UEnum* EnumPtr = EnumType.Get())
 	{
-		int64 CurrentValue = INDEX_NONE;
-		if (const FNumericProperty* NumericProperty = CastField<const FNumericProperty>(Path.GetLastProperty()))
+		if (const void* Container = Path.GetContainerPtr())
 		{
-			if (const void* Container = Path.GetContainerPtr())
+			if (const FNumericProperty* NumericProperty = CastField<const FNumericProperty>(Path.GetLastProperty()))
 			{
-				CurrentValue = NumericProperty->GetSignedIntPropertyValue(NumericProperty->ContainerPtrToValuePtr<const void*>(Container));
+					return NumericProperty->GetSignedIntPropertyValue(NumericProperty->ContainerPtrToValuePtr<const void*>(Container));
 			}
-		}
-		else if (const FEnumProperty* EnumProperty = CastField<const FEnumProperty>(Path.GetLastProperty()))
-		{
-			if (const void* Container = Path.GetContainerPtr())
+			else if (const FEnumProperty* EnumProperty = CastField<const FEnumProperty>(Path.GetLastProperty()))
 			{
-				CurrentValue = EnumProperty->GetUnderlyingProperty()->GetSignedIntPropertyValue(EnumProperty->ContainerPtrToValuePtr<const void*>(Container));
+				return EnumProperty->GetUnderlyingProperty()->GetSignedIntPropertyValue(EnumProperty->ContainerPtrToValuePtr<const void*>(Container));
 			}
 		}
-
-		return CurrentValue;
 	}
 	return 0;
 }
@@ -141,10 +135,10 @@ void SEnumPropertyValue::SetEnumEntry(int32 Index)
 {
 	if (const UEnum* EnumPtr = EnumType.Get())
 	{
-		int64 NewValue = EnumPtr->GetValueByIndex(Index);
-		if (const FNumericProperty* NumericProperty = CastField<const FNumericProperty>(Path.GetLastProperty()))
+		if (void* Container = Path.GetContainerPtr())
 		{
-			if (void* Container = Path.GetContainerPtr())
+			int64 NewValue = EnumPtr->GetValueByIndex(Index);
+			if (const FNumericProperty* NumericProperty = CastField<const FNumericProperty>(Path.GetLastProperty()))
 			{
 				if (NotifyHook)
 				{
@@ -156,20 +150,17 @@ void SEnumPropertyValue::SetEnumEntry(int32 Index)
 					NotifyHook->OnPostValueChange(Path);
 				}
 			}
-		}
-		else if (const FEnumProperty* EnumProperty = CastField<const FEnumProperty>(Path.GetLastProperty()))
-		{
-			if (void* Container = Path.GetContainerPtr())
+			else if (const FEnumProperty* EnumProperty = CastField<const FEnumProperty>(Path.GetLastProperty()))
 			{
 				if (NotifyHook)
 				{
 					NotifyHook->OnPreValueChange(Path);
 				}
 				EnumProperty->GetUnderlyingProperty()->SetIntPropertyValue(EnumProperty->ContainerPtrToValuePtr<const void*>(Container), NewValue);
-				 if (NotifyHook)
-				 {
-					 NotifyHook->OnPostValueChange(Path);
-				 }
+				if (NotifyHook)
+				{
+					NotifyHook->OnPostValueChange(Path);
+				}
 			}
 		}
 	}
diff --git a/Engine/Source/Runtime/AdvancedWidgets/Private/Widgets/PropertyViewer/SStringPropertyValue.cpp b/Engine/Source/Runtime/AdvancedWidgets/Private/Widgets/PropertyViewer/SStringPropertyValue.cpp
--- a/Engine/Source/Runtime/AdvancedWidgets/Private/Widgets/PropertyViewer/SStringPropertyValue.cpp
+++ b/Engine/Source/Runtime/AdvancedWidgets/Private/Widgets/PropertyViewer/SStringPropertyValue.cpp
@@ -46,31 +46,25 @@ void SStringPropertyValue::Construct(const FArguments& InArgs)
 
 FText SStringPropertyValue::GetText() const
 {
-	if (const FStrProperty* StrProperty = CastField<const FStrProperty>(Path.GetLastProperty()))
+	if (const void* Container = Path.GetContainerPtr())
 	{
-		if (const void* Container = Path.GetContainerPtr())
+		if (const FStrProperty* StrProperty = CastField<const FStrProperty>(Path.GetLastProperty()))
 		{
-			FString OutString;
-			StrProperty->GetValue_InContainer(Container, &OutString);
-			return FText::FromString(OutString);
+				FString OutString;
+				StrProperty->GetValue_InContainer(Container, &OutString);
+				return FText::FromString(OutString);
 		}
-	}
-	else if (const FTextProperty* TextProperty = CastField<const FTextProperty>(Path.GetLastProperty()))
-	{
-		if (const void* Container = Path.GetContainerPtr())
+		else if (const FTextProperty* TextProperty = CastField<const FTextProperty>(Path.GetLastProperty()))
 		{
-			FText OutText;
-			TextProperty->GetValue_InContainer(Container, &OutText);
-			return OutText;
+				FText OutText;
+				TextProperty->GetValue_InContainer(Container, &OutText);
+				return OutText;
 		}
-	}
-	else if (const FNameProperty* NameProperty = CastField<const FNameProperty>(Path.GetLastProperty()))
-	{
-		if (const void* Container = Path.GetContainerPtr())
+		else if (const FNameProperty* NameProperty = CastField<const FNameProperty>(Path.GetLastProperty()))
 		{
-			FName OutText;
-			NameProperty->GetValue_InContainer(Container, &OutText);
-			return FText::FromName(OutText);
+				FName OutText;
+				NameProperty->GetValue_InContainer(Container, &OutText);
+				return FText::FromName(OutText);
 		}
 	}
 	return FText::GetEmpty();
@@ -79,9 +73,9 @@ FText SStringPropertyValue::GetText() const
 
 void SStringPropertyValue::OnTextCommitted(const FText& InText, ETextCommit::Type InCommitType)
 {
-	if (const FStrProperty* StrProperty = CastField<const FStrProperty>(Path.GetLastProperty()))
+	if (void* Container = Path.GetContainerPtr())
 	{
-		if (void* Container = Path.GetContainerPtr())
+		if (const FStrProperty* StrProperty = CastField<const FStrProperty>(Path.GetLastProperty()))
 		{
 			FString ToSet = InText.ToString();
 			if (NotifyHook)
@@ -93,11 +87,8 @@ void SStringPropertyValue::OnTextCommitted(const FText& InText, ETextCommit::Typ
 			{
 				NotifyHook->OnPostValueChange(Path);
 			}
-		}
-	}
-	else if (const FTextProperty* TextProperty = CastField<const FTextProperty>(Path.GetLastProperty()))
-	{
-		if (void* Container = Path.GetContainerPtr())
+			}
+		else if (const FTextProperty* TextProperty = CastField<const FTextProperty>(Path.GetLastProperty()))
 		{
 			if (NotifyHook)
 			{
@@ -109,10 +100,7 @@ void SStringPropertyValue::OnTextCommitted(const FText& InText, ETextCommit::Typ
 				NotifyHook->OnPostValueChange(Path);
 			}
 		}
-	}
-	else if (const FNameProperty* NameProperty = CastField<const FNameProperty>(Path.GetLastProperty()))
-	{
-		if (void* Container = Path.GetContainerPtr())
+		else if (const FNameProperty* NameProperty = CastField<const FNameProperty>(Path.GetLastProperty()))
 		{
 			FName ToSet = *(InText.ToString());
 			if (NotifyHook)
