PropertyViewer: use a weakobjptr and fetch the property each time it is requested. That is more slow but prevents crash when an editor compile a BP and doesn't trigger the correct callback.

diff --git a/Engine/Source/Runtime/AdvancedWidgets/Private/Widgets/PropertyViewer/SPropertyViewerImpl.cpp b/Engine/Source/Runtime/AdvancedWidgets/Private/Widgets/PropertyViewer/SPropertyViewerImpl.cpp
--- a/Engine/Source/Runtime/AdvancedWidgets/Private/Widgets/PropertyViewer/SPropertyViewerImpl.cpp
+++ b/Engine/Source/Runtime/AdvancedWidgets/Private/Widgets/PropertyViewer/SPropertyViewerImpl.cpp
@@ -74,14 +74,28 @@ FContainer::FContainer(SPropertyViewer::FHandle InIdentifier, TOptional<FText> I
 }
 
 
+bool InternalIsValid(const UStruct* Container)
+{
+	if (const UClass* Class = Cast<UClass>(Container))
+	{
+		return !Class->HasAnyClassFlags(EClassFlags::CLASS_NewerVersionExists);
+	}
+	else if (const UScriptStruct* ScriptStruct = Cast<UScriptStruct>(Container))
+	{
+		return (ScriptStruct->StructFlags & (EStructFlags::STRUCT_Trashed)) == 0;
+	}
+	return false;
+}
+
 bool FContainer::IsValid() const
 {
-	if (const UClass* Class = Cast<UClass>(Container.Get()))
+	const UStruct * ContainerPtr = Container.Get();
+	if (const UClass* Class = Cast<UClass>(ContainerPtr))
 	{
 		return !Class->HasAnyClassFlags(EClassFlags::CLASS_NewerVersionExists)
 			&& (!bIsObject || (ObjectInstance.Get() && ObjectInstance.Get()->GetClass() == Class));
 	}
-	else if (const UScriptStruct* ScriptStruct = Cast<UScriptStruct>(Container.Get()))
+	else if (const UScriptStruct* ScriptStruct = Cast<UScriptStruct>(ContainerPtr))
 	{
 		return (ScriptStruct->StructFlags & (EStructFlags::STRUCT_Trashed)) == 0;
 	}
@@ -105,7 +119,10 @@ TSharedRef<FTreeNode> FTreeNode::MakeField(TSharedPtr<FTreeNode> InParent, const
 {
 	check(Property);
 	TSharedRef<FTreeNode> Result = MakeShared<FTreeNode>();
-	Result->Property = Property;
+	Result->FieldOwner = Property->GetOwnerStruct();
+	Result->PropertyName = Property->GetFName();
+	Result->bIsStructProperty = CastField<const FStructProperty>(Property) != nullptr;
+	Result->bIsObjectProperty = CastField<const FObjectPropertyBase>(Property) != nullptr;
 	Result->OverrideDisplayName = InDisplayName;
 	Result->ParentNode = InParent;
 	InParent->ChildNodes.Add(Result);
@@ -117,13 +134,30 @@ TSharedRef<FTreeNode> FTreeNode::MakeField(TSharedPtr<FTreeNode> InParent, const
 {
 	check(Function);
 	TSharedRef<FTreeNode> Result = MakeShared<FTreeNode>();
-	Result->Function = Function;
+	Result->FieldOwner = Function->GetOwnerClass();
+	Result->FunctionName = Function->GetFName();
 	Result->OverrideDisplayName = InDisplayName;
 	Result->ParentNode = InParent;
 	InParent->ChildNodes.Add(Result);
 	return Result;
 }
 
+FFieldVariant FTreeNode::GetField() const
+{
+	if (const UStruct* FieldOwnerPtr = FieldOwner.Get())
+	{
+		if (!PropertyName.IsNone())
+		{
+			return FFieldVariant(FieldOwnerPtr->FindPropertyByName(PropertyName));
+		}
+		else if (!FunctionName.IsNone())
+		{
+			return FFieldVariant(CastChecked<const UClass>(FieldOwnerPtr)->FindFunctionByName(FunctionName));
+		}
+	}
+	return FFieldVariant();
+}
+
 
 FPropertyPath FTreeNode::GetPropertyPath() const
 {
@@ -131,9 +165,18 @@ FPropertyPath FTreeNode::GetPropertyPath() const
 	FPropertyPath::FPropertyArray Properties;
 	while (CurrentNode)
 	{
-		if (CurrentNode->Property)
+		const UStruct* CurrentFieldOwner = CurrentNode->FieldOwner.Get();
+		if (CurrentFieldOwner && !CurrentNode->PropertyName.IsNone())
 		{
-			Properties.Insert(CurrentNode->Property, 0);
+			FProperty* Property = CurrentFieldOwner->FindPropertyByName(CurrentNode->PropertyName);
+			if (Property)
+			{
+				Properties.Insert(Property, 0);
+			}
+			else
+			{
+				return FPropertyPath();
+			}
 		}
 
 		TSharedPtr<FContainer> ContainerPin = CurrentNode->Container.Pin();
@@ -220,19 +263,28 @@ TSharedPtr<FContainer> FTreeNode::GetOwnerContainer() const
 
 void FTreeNode::GetFilterStrings(TArray<FString>& OutStrings) const
 {
-	if (Property)
+	const UStruct* CurrentFieldOwner = FieldOwner.Get();
+	if (CurrentFieldOwner && !PropertyName.IsNone())
 	{
-		OutStrings.Add(Property->GetName());
+		FProperty* PropertyPtr = CurrentFieldOwner->FindPropertyByName(PropertyName);
+		if (PropertyPtr)
+		{
+			OutStrings.Add(PropertyPtr->GetName());
 #if WITH_EDITORONLY_DATA
-		OutStrings.Add(Property->GetDisplayNameText().ToString());
+			OutStrings.Add(PropertyPtr->GetDisplayNameText().ToString());
 #endif
+		}
 	}
-	if (const UFunction* FunctionPtr = Function.Get())
+	if (CurrentFieldOwner && !FunctionName.IsNone())
 	{
-		OutStrings.Add(FunctionPtr->GetName());
+		const UFunction* FunctionPtr = CastChecked<const UClass>(CurrentFieldOwner)->FindFunctionByName(FunctionName);
+		if (FunctionPtr)
+		{
+			OutStrings.Add(FunctionPtr->GetName());
 #if WITH_EDITORONLY_DATA
-		OutStrings.Add(FunctionPtr->GetDisplayNameText().ToString());
+			OutStrings.Add(FunctionPtr->GetDisplayNameText().ToString());
 #endif
+		}
 	}
 	if (const TSharedPtr<FContainer> ContainerPtr = Container.Pin())
 	{
@@ -270,16 +322,19 @@ void FTreeNode::BuildChildNodesRecursive(IFieldIterator& FieldIterator, IFieldEx
 	ChildNodes.Reset();
 
 	const UStruct* ChildStructType = nullptr;
-	if (Property)
+	if (bIsStructProperty || bIsObjectProperty)
 	{
-		if (const FStructProperty* StructProperty = CastField<const FStructProperty>(Property))
+		check(!PropertyName.IsNone());
+		const UStruct* CurrentFieldOwner = FieldOwner.Get();
+		FProperty* PropertyPtr = CurrentFieldOwner ? CurrentFieldOwner->FindPropertyByName(PropertyName) : nullptr;
+		if (const FStructProperty* StructProperty = CastField<const FStructProperty>(PropertyPtr))
 		{
 			if (FieldExpander.CanExpandScriptStruct(StructProperty))
 			{
 				ChildStructType = StructProperty->Struct;
 			}
 		}
-		else if (const FObjectPropertyBase* ObjectProperty = CastField<const FObjectPropertyBase>(Property))
+		else if (const FObjectPropertyBase* ObjectProperty = CastField<const FObjectPropertyBase>(PropertyPtr))
 		{
 			UObject* Instance = nullptr;
 			if (TSharedPtr<FContainer> OwnerContainer = GetOwnerContainer())
@@ -305,12 +360,17 @@ void FTreeNode::BuildChildNodesRecursive(IFieldIterator& FieldIterator, IFieldEx
 			}
 		}
 	}
-	else if (const UFunction* FunctionPtr = Function.Get())
+	else if (!FunctionName.IsNone())
 	{
-		TOptional<const UStruct*> StructToExpand = FieldExpander.GetExpandedFunction(FunctionPtr);
-		if (StructToExpand.IsSet())
+		const UStruct* CurrentFieldOwner = FieldOwner.Get();
+		UFunction* FunctionPtr = CurrentFieldOwner ? CastChecked<const UClass>(CurrentFieldOwner)->FindFunctionByName(FunctionName) : nullptr;
+		if (FunctionPtr)
 		{
-			ChildStructType = StructToExpand.GetValue();
+			TOptional<const UStruct*> StructToExpand = FieldExpander.GetExpandedFunction(FunctionPtr);
+			if (StructToExpand.IsSet())
+			{
+				ChildStructType = StructToExpand.GetValue();
+			}
 		}
 	}
 	else if (const TSharedPtr<FContainer> ContainerPin = Container.Pin())
@@ -348,10 +408,10 @@ bool FTreeNode::Sort(const TSharedPtr<FTreeNode>& NodeA, const TSharedPtr<FTreeN
 {
 	bool bIsContainerA = NodeA->IsContainer();
 	bool bIsContainerB = NodeB->IsContainer();
-	bool bIsObjectPropertyA = CastField<FObjectPropertyBase>(NodeA->Property) != nullptr;
-	bool bIsObjectPropertyB = CastField<FObjectPropertyBase>(NodeB->Property) != nullptr;
-	bool bIsFunctionA = NodeA->Function.Get() != nullptr;
-	bool bIsFunctionB = NodeB->Function.Get() != nullptr;
+	bool bIsObjectPropertyA = NodeA->bIsObjectProperty;
+	bool bIsObjectPropertyB = NodeB->bIsObjectProperty;
+	bool bIsFunctionA = !NodeA->FunctionName.IsNone();
+	bool bIsFunctionB = !NodeB->FunctionName.IsNone();
 	const FName NodeStrA = bIsContainerA ? NodeA->GetContainer()->GetStruct()->GetFName() : NodeA->GetField().GetFName();
 	const FName NodeStrB = bIsContainerB ? NodeB->GetContainer()->GetStruct()->GetFName() : NodeB->GetField().GetFName();
 
diff --git a/Engine/Source/Runtime/AdvancedWidgets/Private/Widgets/PropertyViewer/SPropertyViewerImpl.h b/Engine/Source/Runtime/AdvancedWidgets/Private/Widgets/PropertyViewer/SPropertyViewerImpl.h
--- a/Engine/Source/Runtime/AdvancedWidgets/Private/Widgets/PropertyViewer/SPropertyViewerImpl.h
+++ b/Engine/Source/Runtime/AdvancedWidgets/Private/Widgets/PropertyViewer/SPropertyViewerImpl.h
@@ -117,8 +117,11 @@ private:
 
 	//Either a container, property or function
 	TWeakPtr<FContainer> Container;
-	const FProperty* Property = nullptr;
-	TWeakObjectPtr<const UFunction> Function;
+	TWeakObjectPtr<const UStruct> FieldOwner;
+	FName PropertyName;
+	FName FunctionName;
+	bool bIsStructProperty = false;
+	bool bIsObjectProperty = false;
 
 	TOptional<FText> OverrideDisplayName;
 
@@ -140,13 +143,10 @@ public:
 
 	bool IsField() const
 	{
-		return Property != nullptr  || Function.Get() != nullptr;
+		return !PropertyName.IsNone() || !FunctionName.IsNone();
 	}
 	
-	FFieldVariant GetField() const
-	{
-		return Property != nullptr ? FFieldVariant(Property) : Function.Get() ? FFieldVariant(Function.Get()) : FFieldVariant();
-	}
+	FFieldVariant GetField() const;
 
 	TSharedPtr<FTreeNode> GetParentNode() const
 	{
