[AdvancedWidgets]
* Removed dll export on enums

diff --git a/Engine/Source/Runtime/AdvancedWidgets/Public/Framework/ColorGrading/ColorGradingCommon.h b/Engine/Source/Runtime/AdvancedWidgets/Public/Framework/ColorGrading/ColorGradingCommon.h
--- a/Engine/Source/Runtime/AdvancedWidgets/Public/Framework/ColorGrading/ColorGradingCommon.h
+++ b/Engine/Source/Runtime/AdvancedWidgets/Public/Framework/ColorGrading/ColorGradingCommon.h
@@ -12,7 +12,7 @@ namespace UE::ColorGrading
 DECLARE_DELEGATE_RetVal_OneParam(bool, FOnGetCurrentVector4Value, FVector4&);
 
 /** Color components that can be controlled by a color grading widget. */
-enum class ADVANCEDWIDGETS_API EColorGradingComponent
+enum class EColorGradingComponent
 {
 	Red,
 	Green,
@@ -24,7 +24,7 @@ enum class ADVANCEDWIDGETS_API EColorGradingComponent
 };
 
 /** Types of color grading modes */
-enum class ADVANCEDWIDGETS_API EColorGradingWheelType
+enum class EColorGradingWheelType
 {
 	/** A wheel that controls a standard color value */
 	Standard,
@@ -34,14 +34,14 @@ enum class ADVANCEDWIDGETS_API EColorGradingWheelType
 };
 
 /** Modes that colors can be displayed in for the color grading panel */
-enum class ADVANCEDWIDGETS_API EColorGradingColorDisplayMode
+enum class EColorGradingColorDisplayMode
 {
 	RGB,
 	HSV
 };
 
 /** Enumerates color picker modes */
-enum class ADVANCEDWIDGETS_API EColorGradingModes
+enum class EColorGradingModes
 {
 	Saturation,
 	Contrast,
