[AdvancedWidgets] Fix color grading wheel selector jumping to wrong location when clicked

diff --git a/Engine/Source/Runtime/AdvancedWidgets/Private/Widgets/ColorGrading/SColorGradingWheel.cpp b/Engine/Source/Runtime/AdvancedWidgets/Private/Widgets/ColorGrading/SColorGradingWheel.cpp
--- a/Engine/Source/Runtime/AdvancedWidgets/Private/Widgets/ColorGrading/SColorGradingWheel.cpp
+++ b/Engine/Source/Runtime/AdvancedWidgets/Private/Widgets/ColorGrading/SColorGradingWheel.cpp
@@ -78,6 +78,22 @@ void SColorGradingWheel::SetExponentDisplacementAttribute(TAttribute<float> InEx
 	}
 }
 
+FVector2f SColorGradingWheel::GetActualSize(const FGeometry& MyGeometry) const
+{
+	const FVector2f AllottedGeometrySize = MyGeometry.GetLocalSize();
+
+	if (bIsAttributeDesiredWheelSizeSet)
+	{
+		// Even if a desired size is provided, make sure the wheel is painted within the allotted geometry
+		int32 CachedDesiredWheelSize = DesiredWheelSizeAttribute.Get();
+		float ActualSize = FMath::Min(CachedDesiredWheelSize, AllottedGeometrySize.GetMin());
+
+		return FVector2f(ActualSize, ActualSize);
+	}
+
+	return AllottedGeometrySize;
+}
+
 FString SColorGradingWheel::GetReferencerName() const
 {
 	return TEXT("SColorGradingWheel");
@@ -171,19 +187,10 @@ int32 SColorGradingWheel::OnPaint(const FPaintArgs& Args, const FGeometry& Allot
 	const bool bIsEnabled = ShouldBeEnabled(bParentEnabled);
 	const ESlateDrawEffect DrawEffects = bIsEnabled ? ESlateDrawEffect::None : ESlateDrawEffect::DisabledEffect;
 
-	FVector2f AllottedGeometrySize = AllottedGeometry.GetLocalSize();
-
-	if (bIsAttributeDesiredWheelSizeSet)
-	{
-		// Even if a desired size is provided, make sure the wheel is painted within the allotted geometry
-		int32 CachedDesiredWheelSize = DesiredWheelSizeAttribute.Get();
-		float ActualSize = FMath::Min(CachedDesiredWheelSize, AllottedGeometrySize.GetMin());
-
-		AllottedGeometrySize = FVector2f(ActualSize, ActualSize);
-	}
+	FVector2f GeometrySize = GetActualSize(AllottedGeometry);
 
 	const FVector2f SelectorSize = SelectorImage->ImageSize;
-	const FVector2f CircleSize = AllottedGeometrySize - SelectorSize;
+	const FVector2f CircleSize = GeometrySize - SelectorSize;
 
 	FSlateDrawElement::MakeBox(
 		OutDrawElements,
@@ -206,7 +213,7 @@ int32 SColorGradingWheel::OnPaint(const FPaintArgs& Args, const FGeometry& Allot
 	FSlateDrawElement::MakeBox(
 		OutDrawElements,
 		LayerId + 2,
-		AllottedGeometry.ToPaintGeometry(SelectorSize, FSlateLayoutTransform(0.5f * (AllottedGeometrySize + CalcRelativePositionFromCenter() * CircleSize - SelectorSize))),
+		AllottedGeometry.ToPaintGeometry(SelectorSize, FSlateLayoutTransform(0.5f * (GeometrySize + CalcRelativePositionFromCenter() * CircleSize - SelectorSize))),
 		SelectorImage,
 		DrawEffects,
 		InWidgetStyle.GetColorAndOpacityTint() * SelectorImage->GetTint(InWidgetStyle)
@@ -237,15 +244,13 @@ UE::Slate::FDeprecateVector2DResult SColorGradingWheel::CalcRelativePositionFrom
 
 bool SColorGradingWheel::ProcessMouseAction(const FGeometry& MyGeometry, const FPointerEvent& MouseEvent, bool bProcessWhenOutsideColorWheel)
 {
-	FVector2f GeometrySize = MyGeometry.GetLocalSize();
 	if (bIsAttributeDesiredWheelSizeSet)
 	{
 		DesiredWheelSizeAttribute.UpdateNow(*this);
-		int32 CachedDesiredWheelSize = DesiredWheelSizeAttribute.Get();
-		GeometrySize.X = CachedDesiredWheelSize;
-		GeometrySize.Y = CachedDesiredWheelSize;
 	}
 
+	const FVector2f GeometrySize = GetActualSize(MyGeometry);
+
 	const FVector2f LocalMouseCoordinate = MyGeometry.AbsoluteToLocal(MouseEvent.GetScreenSpacePosition());
 	const FVector2f RelativePositionFromCenter = (2.0f * LocalMouseCoordinate - GeometrySize) / (GeometrySize - SelectorImage->ImageSize);
 	const float RelativeRadius = RelativePositionFromCenter.Size();
diff --git a/Engine/Source/Runtime/AdvancedWidgets/Public/Widgets/ColorGrading/SColorGradingWheel.h b/Engine/Source/Runtime/AdvancedWidgets/Public/Widgets/ColorGrading/SColorGradingWheel.h
--- a/Engine/Source/Runtime/AdvancedWidgets/Public/Widgets/ColorGrading/SColorGradingWheel.h
+++ b/Engine/Source/Runtime/AdvancedWidgets/Public/Widgets/ColorGrading/SColorGradingWheel.h
@@ -100,6 +100,9 @@ protected:
 	/** */
 	ADVANCEDWIDGETS_API void SetExponentDisplacementAttribute(TAttribute<float> InExponentDisplacement);
 
+	/** Get the actual size of the wheel, taking into account the available space */
+	ADVANCEDWIDGETS_API FVector2f GetActualSize(const FGeometry& MyGeometry) const;
+
 	/** @return an attribute reference of SelectedColor */
 	TSlateAttributeRef<FLinearColor> GetSelectedColorAttribute() const { return TSlateAttributeRef<FLinearColor>(SharedThis(this), SelectedColorAttribute); }
 
