PropertyViewer: Add the ability focus the search box when the viewer/picker is created

For use in menus

diff --git a/Engine/Source/Runtime/AdvancedWidgets/Private/Widgets/PropertyViewer/SPropertyViewerImpl.cpp b/Engine/Source/Runtime/AdvancedWidgets/Private/Widgets/PropertyViewer/SPropertyViewerImpl.cpp
--- a/Engine/Source/Runtime/AdvancedWidgets/Private/Widgets/PropertyViewer/SPropertyViewerImpl.cpp
+++ b/Engine/Source/Runtime/AdvancedWidgets/Private/Widgets/PropertyViewer/SPropertyViewerImpl.cpp
@@ -549,6 +549,26 @@ TSharedRef<SWidget> FPropertyViewerImpl::Construct(const SPropertyViewer::FArgum
 				[
 					CreateSearch()
 				];
+
+			if (InArgs._bFocusSearchBox)
+			{
+				SearchBox->RegisterActiveTimer( 0.f, FWidgetActiveTimerDelegate::CreateSPLambda(this, [this](double InCurrentTime, float InDeltaTime)
+				{
+					if (SearchBoxWidget.IsValid())
+					{
+						FWidgetPath WidgetToFocusPath;
+						FSlateApplication::Get().GeneratePathToWidgetUnchecked(SearchBoxWidget.ToSharedRef(), WidgetToFocusPath);
+						if(WidgetToFocusPath.IsValid())
+						{
+							FSlateApplication::Get().SetKeyboardFocus(WidgetToFocusPath, EFocusCause::SetDirectly);
+							WidgetToFocusPath.GetWindow()->SetWidgetToFocusOnActivate(SearchBoxWidget);
+						}
+						return EActiveTimerReturnType::Stop;
+					}
+
+					return EActiveTimerReturnType::Continue;
+				}));
+			}
 		}
 		else
 		{
diff --git a/Engine/Source/Runtime/AdvancedWidgets/Public/Widgets/PropertyViewer/SPropertyViewer.h b/Engine/Source/Runtime/AdvancedWidgets/Public/Widgets/PropertyViewer/SPropertyViewer.h
--- a/Engine/Source/Runtime/AdvancedWidgets/Public/Widgets/PropertyViewer/SPropertyViewer.h
+++ b/Engine/Source/Runtime/AdvancedWidgets/Public/Widgets/PropertyViewer/SPropertyViewer.h
@@ -87,6 +87,8 @@ public:
 		SLATE_ARGUMENT_DEFAULT(bool, bSortChildNode) = false;
 		/** Show a search box. */
 		SLATE_ARGUMENT_DEFAULT(bool, bShowSearchBox) = false;
+		/** Focus the search box when the widget is first displayed. */
+		SLATE_ARGUMENT_DEFAULT(bool, bFocusSearchBox) = false;
 		/** Is selection enabled. */ 
 		SLATE_ARGUMENT_DEFAULT(ESelectionMode::Type, SelectionMode) = ESelectionMode::Single;
 		/** Which properties/functions to show. FFieldIterator_BlueprintVisible is the default. */
