SColorGradingComponentSpinBox, SColorSlider, SComplexGradient : check that NumColors > 1 for gradient to avoid divide by zero
DrawCentripetalCatmullRomSpline : check that NumColors > 0 to avoid using a negative array index in the Colors array

diff --git a/Engine/Plugins/Animation/PoseSearch/Source/Runtime/Private/PoseSearchContext.cpp b/Engine/Plugins/Animation/PoseSearch/Source/Runtime/Private/PoseSearchContext.cpp
--- a/Engine/Plugins/Animation/PoseSearch/Source/Runtime/Private/PoseSearchContext.cpp
+++ b/Engine/Plugins/Animation/PoseSearch/Source/Runtime/Private/PoseSearchContext.cpp
@@ -638,7 +638,7 @@ void FDebugDrawParams::DrawCentripetalCatmullRomSpline(TConstArrayView<FVector>
 {
 	const int32 NumPoints = Points.Num();
 	const int32 NumColors = Colors.Num();
-	if (NumPoints > 1)
+	if (NumPoints > 1 && NumColors > 0)
 	{
 		auto GetT = [](float T, float Alpha, const FVector& P0, const FVector& P1)
 		{
diff --git a/Engine/Source/Runtime/AdvancedWidgets/Private/Widgets/ColorGrading/SColorGradingComponentSpinBox.cpp b/Engine/Source/Runtime/AdvancedWidgets/Private/Widgets/ColorGrading/SColorGradingComponentSpinBox.cpp
--- a/Engine/Source/Runtime/AdvancedWidgets/Private/Widgets/ColorGrading/SColorGradingComponentSpinBox.cpp
+++ b/Engine/Source/Runtime/AdvancedWidgets/Private/Widgets/ColorGrading/SColorGradingComponentSpinBox.cpp
@@ -100,7 +100,7 @@ int32 SColorGradingComponentSpinBox::OnPaint(const FPaintArgs& Args, const FGeom
 		const TArray<FLinearColor> Colors = GetGradientColors();
 		int32 NumColors = Colors.Num();
 
-		if (NumColors > 0)
+		if (NumColors > 1)
 		{
 			TArray<FSlateGradientStop> GradientStops;
 
diff --git a/Engine/Source/Runtime/AppFramework/Private/Widgets/Colors/SColorSlider.cpp b/Engine/Source/Runtime/AppFramework/Private/Widgets/Colors/SColorSlider.cpp
--- a/Engine/Source/Runtime/AppFramework/Private/Widgets/Colors/SColorSlider.cpp
+++ b/Engine/Source/Runtime/AppFramework/Private/Widgets/Colors/SColorSlider.cpp
@@ -146,7 +146,7 @@ int32 SColorSlider::OnPaint(const FPaintArgs& Args, const FGeometry& AllottedGeo
 	// Draw Color Gradient
 	const int32 NumColors = GradientColors.Get().Num();
 
-	if (NumColors > 0)
+	if (NumColors > 1)
 	{
 		if (bHasAlphaBackground.Get())
 		{
diff --git a/Engine/Source/Runtime/Engine/Private/DrawDebugHelpers.cpp b/Engine/Source/Runtime/Engine/Private/DrawDebugHelpers.cpp
--- a/Engine/Source/Runtime/Engine/Private/DrawDebugHelpers.cpp
+++ b/Engine/Source/Runtime/Engine/Private/DrawDebugHelpers.cpp
@@ -927,7 +927,7 @@ void DrawCentripetalCatmullRomSpline(const UWorld* InWorld, TConstArrayView<FVec
 	
 	const int32 NumPoints = Points.Num();
 	const int32 NumColors = Colors.Num();
-	if (NumPoints > 1)
+	if (NumPoints > 1 && NumColors > 0)
 	{
 		auto GetT = [](float T, float Alpha, const FVector& P0, const FVector& P1)
 		{
diff --git a/Engine/Source/Runtime/Slate/Private/Widgets/Colors/SComplexGradient.cpp b/Engine/Source/Runtime/Slate/Private/Widgets/Colors/SComplexGradient.cpp
--- a/Engine/Source/Runtime/Slate/Private/Widgets/Colors/SComplexGradient.cpp
+++ b/Engine/Source/Runtime/Slate/Private/Widgets/Colors/SComplexGradient.cpp
@@ -51,7 +51,7 @@ int32 SComplexGradient::OnPaint( const FPaintArgs& Args, const FGeometry& Allott
 	const TArray<FLinearColor>& Colors = GradientColors.Get();
 	int32 NumColors = Colors.Num();
 
-	if (NumColors > 0)
+	if (NumColors > 1)
 	{
 		TArray<FSlateGradientStop> GradientStops;
 
