UMG Editor: Viewmodel panel upgrades.
- You can now call functions, edit name/string/text values and change object properties.
- Fixed viewmodel panel not displaying properties correctly when first opened if an object is already set to be debugged.

#Synced-CL 52492987

diff --git a/Engine/Plugins/Runtime/ModelViewViewModel/Source/ModelViewViewModelEditor/Private/Widgets/SMVVMViewModelPanel.cpp b/Engine/Plugins/Runtime/ModelViewViewModel/Source/ModelViewViewModelEditor/Private/Widgets/SMVVMViewModelPanel.cpp
--- a/Engine/Plugins/Runtime/ModelViewViewModel/Source/ModelViewViewModelEditor/Private/Widgets/SMVVMViewModelPanel.cpp
+++ b/Engine/Plugins/Runtime/ModelViewViewModel/Source/ModelViewViewModelEditor/Private/Widgets/SMVVMViewModelPanel.cpp
@@ -271,13 +271,11 @@ void SMVVMViewModelPanel::Construct(const FArguments& InArgs, TSharedPtr<FWidget
 		]
 	];
 
-	CreateViewModelTreeView(UE::PropertyViewer::SPropertyViewer::EPropertyVisibility::Visible);
-
-	FillViewModel();
+	HandleSetObjectBeingDebugged(WidgetBlueprint->GetObjectBeingDebugged());
 }
 
 SMVVMViewModelPanel::SMVVMViewModelPanel()
-	: PropertyVisibility(UE::PropertyViewer::SPropertyViewer::EPropertyVisibility::Visible)
+	: PropertyVisibility(UE::PropertyViewer::SPropertyViewer::EPropertyVisibility::Hidden)
 {}
 
 SMVVMViewModelPanel::~SMVVMViewModelPanel()
@@ -361,7 +359,8 @@ void SMVVMViewModelPanel::FillViewModel()
 						}
 					}
 				}
-				else
+				
+				if (!Handle.IsValid())
 				{
 					Handle = ViewModelTreeView->AddContainer(ViewModelClass);
 				}
@@ -376,7 +375,13 @@ void SMVVMViewModelPanel::FillViewModel()
 		// Only add widget values when in editable mode.
 		if (Widget && PropertyVisibility == UE::PropertyViewer::SPropertyViewer::EPropertyVisibility::Editable)
 		{
-			ViewModelTreeView->AddInstance(Widget);
+			UE::PropertyViewer::SPropertyViewer::FHandle Handle = ViewModelTreeView->AddInstance(Widget);
+
+			if (Handle.IsValid())
+			{
+				// Empty GUID for the widget instance.
+				PropertyViewerHandles.Add(Handle, FGuid());
+			}
 		}
 	}
 }
@@ -618,13 +623,46 @@ TSharedRef<SWidget> SMVVMViewModelPanel::HandleGenerateContainer(UE::PropertyVie
 {
 	if (const FGuid* VMGuidPtr = PropertyViewerHandles.Find(ContainerHandle))
 	{
+		UWidgetBlueprint* WidgetBlueprint = nullptr;
+
+		if (TSharedPtr<FWidgetBlueprintEditor> WidgetBlueprintEditor = WeakBlueprintEditor.Pin())
+		{
+			WidgetBlueprint = WidgetBlueprintEditor->GetWidgetBlueprintObj();
+		}
+
 		const FGuid VMGuid = *VMGuidPtr;
+		// Widget itself, not a view model
+		if (!VMGuid.IsValid())
+		{
+			if (WidgetBlueprint)
+			{
+				return SNew(STextBlock)
+					.Text(FText::FromString(WidgetBlueprint->GetName()));
+			}
+		}
+
 		if (UMVVMBlueprintView* View = WeakBlueprintView.Get())
 		{
 			if (const FMVVMBlueprintViewModelContext* ViewModelContext = View->FindViewModel(VMGuid))
 			{
 				if (ViewModelContext->GetViewModelClass())
 				{
+					bool bReadOnly = true;
+
+					if (WidgetBlueprint)
+					{
+						if (UUserWidget* Widget = Cast<UUserWidget>(WidgetBlueprint->GetObjectBeingDebugged()))
+						{
+							if (UMVVMView* WidgetView = UMVVMSubsystem::GetViewFromUserWidget(Widget))
+							{
+								if (WidgetView->GetViewModel(ViewModelContext->GetViewModelName()).GetObject())
+								{
+									bReadOnly = false;
+								}
+							}
+						}
+					}
+
 					TSharedRef<SInlineEditableTextBlock> EditableTextBlock = SNew(SInlineEditableTextBlock)
 						.Text(ViewModelContext->GetDisplayName())
 						.IsReadOnly_Lambda([this, VMGuid]() { return !HandleCanRename(VMGuid); })
@@ -632,7 +670,7 @@ TSharedRef<SWidget> SMVVMViewModelPanel::HandleGenerateContainer(UE::PropertyVie
 						.OnTextCommitted(this, &SMVVMViewModelPanel::HandleNameTextCommited, VMGuid);
 					EditableTextBlocks.Add(VMGuid, EditableTextBlock);
 
-					return SNew(SHorizontalBox)
+					TSharedRef<SHorizontalBox> Widget = SNew(SHorizontalBox)
 						+ SHorizontalBox::Slot()
 						.AutoWidth()
 						.HAlign(HAlign_Right)
@@ -642,9 +680,25 @@ TSharedRef<SWidget> SMVVMViewModelPanel::HandleGenerateContainer(UE::PropertyVie
 						]
 						+ SHorizontalBox::Slot()
 						.Padding(4.0f)
+						.AutoWidth()
+						.VAlign(VAlign_Center)
 						[
 							EditableTextBlock
 						];
+
+					if (bReadOnly)
+					{
+						Widget->AddSlot()
+							.AutoWidth()
+							.HAlign(HAlign_Left)
+							.VAlign(VAlign_Center)
+							[
+								SNew(STextBlock)
+								.Text(LOCTEXT("ReadOnly", "[Read Only]"))
+							];
+					}
+
+					return Widget;
 				}
 				else
 				{
@@ -1050,7 +1104,7 @@ void SMVVMViewModelPanel::CreateViewModelTreeView(UE::PropertyViewer::SPropertyV
 {
 	PropertyVisibility = InPropertyVisibility;
 
-	if (!ToolbarWidget.IsValid() || SettingsWidget.IsValid())
+	if (!ToolbarWidget.IsValid() || !SettingsWidget.IsValid())
 	{
 		FToolMenuContext GenerateWidgetContext;
 		{
diff --git a/Engine/Source/Runtime/AdvancedWidgets/AdvancedWidgets.Build.cs b/Engine/Source/Runtime/AdvancedWidgets/AdvancedWidgets.Build.cs
--- a/Engine/Source/Runtime/AdvancedWidgets/AdvancedWidgets.Build.cs
+++ b/Engine/Source/Runtime/AdvancedWidgets/AdvancedWidgets.Build.cs
@@ -27,6 +27,7 @@ public class AdvancedWidgets : ModuleRules
 		if (Target.bBuildEditor == true)
 		{
 			PrivateDependencyModuleNames.Add("UnrealEd");
+			PrivateDependencyModuleNames.Add("PropertyEditor");
 		}
 
 		bDisableAutoRTFMInstrumentation = true;
diff --git a/Engine/Source/Runtime/AdvancedWidgets/Private/Framework/PropertyViewer/PropertyValueFactory.cpp b/Engine/Source/Runtime/AdvancedWidgets/Private/Framework/PropertyViewer/PropertyValueFactory.cpp
--- a/Engine/Source/Runtime/AdvancedWidgets/Private/Framework/PropertyViewer/PropertyValueFactory.cpp
+++ b/Engine/Source/Runtime/AdvancedWidgets/Private/Framework/PropertyViewer/PropertyValueFactory.cpp
@@ -1,17 +1,27 @@
 // Copyright Epic Games, Inc. All Rights Reserved.
 
 #include "Framework/PropertyViewer/PropertyValueFactory.h"
+#include "HAL/PlatformMemory.h"
+#include "Misc/MessageDialog.h"
 #include "UObject/EnumProperty.h"
 #include "UObject/TextProperty.h"
 #include "UObject/TopLevelAssetPath.h"
+#include "UObject/StrProperty.h"
+#include "UObject/TextProperty.h"
+#include "UObject/Object.h"
 #include "UObject/UnrealType.h"
 
+#include "Widgets/Input/SButton.h"
 #include "Widgets/PropertyViewer/SBoolPropertyValue.h"
 #include "Widgets/PropertyViewer/SEnumPropertyValue.h"
 #include "Widgets/PropertyViewer/SDefaultPropertyValue.h"
 #include "Widgets/PropertyViewer/SNumericPropertyValue.h"
 #include "Widgets/PropertyViewer/SStringPropertyValue.h"
+#include "Widgets/PropertyViewer/SObjectPropertyValue.h"
+#include "Widgets/SNullWidget.h"
+#include "Widgets/Text/STextBlock.h"
 
+#define LOCTEXT_NAMESPACE "PropertyValueFactory"
 
 namespace UE::PropertyViewer
 {
@@ -59,11 +69,212 @@ FPropertyValueFactory::FPropertyValueFactory()
 	Register(FStrProperty::StaticClass(), FPropertyValueFactory::FOnGenerate::CreateStatic(SStringPropertyValue::CreateInstance));
 	Register(FTextProperty::StaticClass(), FPropertyValueFactory::FOnGenerate::CreateStatic(SStringPropertyValue::CreateInstance));
 	Register(FNameProperty::StaticClass(), FPropertyValueFactory::FOnGenerate::CreateStatic(SStringPropertyValue::CreateInstance));
+	Register(UObject::StaticClass(), FPropertyValueFactory::FOnGenerate::CreateStatic(SObjectPropertyValue::CreateInstance));
+}
+
+FProperty* GetReturnProperty(UFunction* InFunction)
+{
+	if (FProperty* ReturnProperty = InFunction->GetReturnProperty())
+	{
+		return ReturnProperty;
+	}
+
+	if (InFunction->HasAllFunctionFlags(FUNC_HasOutParms))
+	{
+		for (TFieldIterator<FProperty> It(InFunction); It && (It->PropertyFlags & CPF_Parm); ++It)
+		{
+			if (It->HasAllPropertyFlags(CPF_OutParm) && !It->HasAnyPropertyFlags(CPF_Deprecated | CPF_EditorOnly | CPF_ConstParm | CPF_ReferenceParm))
+			{
+				return *It;
+			}
+		}
+	}
+
+	return nullptr;
 }
 
+UObject* GetUObjectFromContainer(void* Container)
+{
+	UObject* Object = static_cast<UObject*>(Container);
+
+	if (IsValid(Object))
+	{
+		return Object;
+	}
+
+	return nullptr;
+}
+
+FText GetDebugPropertyDisplayText(const FProperty* InProperty, void* DataPtr)
+{
+	if (!InProperty)
+	{
+		return FText::GetEmpty();
+	}
+
+	if (const FObjectPropertyBase* ObjectProperty = CastField<const FObjectPropertyBase>(InProperty))
+	{
+		UObject* Value = nullptr;
+		ObjectProperty->GetValue_InContainer(DataPtr, &Value);
+		
+		if (Value)
+		{
+			return FText::FromString(Value->GetPathName());
+		}
+
+		return LOCTEXT("Nullptr", "None");
+	}
+	else if (const FNameProperty* NameProperty = CastField<const FNameProperty>(InProperty))
+	{
+		FName Value = NAME_None;
+		NameProperty->GetValue_InContainer(DataPtr, &Value);
+
+		return FText::FromName(Value);
+	}
+	else if (const FStrProperty* StringProperty = CastField<const FStrProperty>(InProperty))
+	{
+		FString Value = TEXT("");
+		StringProperty->GetValue_InContainer(DataPtr, &Value);
+
+		return FText::Format(
+			LOCTEXT("QuotedString", "\"{0}\""),
+			FText::FromString(Value)
+		);
+	}
+	else if (const FTextProperty* TextProperty = CastField<const FTextProperty>(InProperty))
+	{
+		FText Value = FText::GetEmpty();
+		TextProperty->GetValue_InContainer(DataPtr, &Value);
+
+		return FText::Format(
+			LOCTEXT("QuotedString", "\"{0}\""),
+			Value
+		);
+	}
+
+	FString Value;
+	InProperty->ExportTextItem_InContainer(Value, DataPtr, DataPtr, nullptr, PPF_None, nullptr);
+
+	return FText::FromString(Value);
+}
+
+FReply HandleCallFunctionPressed(TWeakObjectPtr<UClass> ClassWeak, TWeakObjectPtr<UObject> ContainerWeak, TWeakObjectPtr<UFunction> FunctionWeak)
+{
+	UClass* Class = ClassWeak.Get();
+
+	if (!Class)
+	{
+		return FReply::Handled();
+	}
+
+	UObject* Container = ContainerWeak.Get();
+
+	if (!Container)
+	{
+		return FReply::Handled();
+	}
+
+	UFunction* Function = FunctionWeak.Get();
+
+	if (!Function)
+	{
+		return FReply::Handled();
+	}
+
+	if (FProperty* ReturnProperty = GetReturnProperty(Function))
+	{
+		void* DataPtr = FMemory_Alloca_Aligned(Function->ParmsSize, Function->MinAlignment);
+		void* ReturnPropertyPtr = (uint8*)DataPtr + ReturnProperty->GetOffset_ForUFunction();
+
+		ReturnProperty->InitializeValue(ReturnPropertyPtr);
+
+		Container->ProcessEvent(Function, DataPtr);
+
+		const FText DisplayValue = GetDebugPropertyDisplayText(ReturnProperty, DataPtr);
+
+		ReturnProperty->DestroyValue(ReturnPropertyPtr);
+
+#if WITH_EDITOR
+		FMessageDialog::Open(
+			EAppMsgType::Ok,
+			FText::Format(
+				LOCTEXT("ReturnValueFormat", "{0} :: {1}\n\n{2}"),
+				Class->GetDisplayNameText(),
+				Function->GetDisplayNameText(),
+				DisplayValue
+			)
+		);
+#endif
+	}
+	else
+	{
+		Container->ProcessEvent(Function, nullptr);
+	}
+
+	return FReply::Handled();
+}
 
 TSharedPtr<SWidget> FPropertyValueFactory::Generate(FGenerateArgs Args) const
 {
+	if (UFunction* BaseFunction = Args.Field ? Cast<UFunction>(Args.Field->ToUObject()) : nullptr)
+	{
+		// Attempt to resolve property chain, but do not call functions to continue the chain.
+		if (void* ContainerPtr = Args.Path.GetContainerPtr())
+		{
+			if (UObject* Container = GetUObjectFromContainer(ContainerPtr))
+			{
+				UClass* ContainerClass = Container->GetClass();
+
+				if (const FProperty* LastProperty = Args.Path.GetLastProperty())
+				{
+					UClass* PropertyOwnerClass = LastProperty->GetOwnerClass();
+
+					if (ContainerClass != PropertyOwnerClass && !ContainerClass->IsChildOf(PropertyOwnerClass))
+					{
+						return SNew(STextBlock)
+							.Text(LOCTEXT("InvalidClass", "Invalid Class"))
+							.IsEnabled(false);
+					}
+
+					if (const FObjectPropertyBase* ObjectProperty = CastField<const FObjectPropertyBase>(LastProperty))
+					{
+						UObject* NewContainer = nullptr;
+						ObjectProperty->GetValue_InContainer(Container, &NewContainer);
+						Container = NewContainer;
+					}
+					else
+					{
+						return SNew(STextBlock)
+							.Text(LOCTEXT("InvalidProperty", "Invalid Property"))
+							.IsEnabled(false);
+					}
+
+					if (!IsValid(Container))
+					{
+						return SNew(STextBlock)
+							.Text(LOCTEXT("NoParentObject", "No Parent Object"))
+							.IsEnabled(false);
+					}
+				}
+
+				// Ensure the container actually has the function.
+				if (UFunction* Function = Container->FindFunction(BaseFunction->GetFName()))
+				{
+					FProperty* ReturnProperty = GetReturnProperty(Function);
+
+					if ((Function->NumParms == 0 && !ReturnProperty) || (Function->NumParms == 1 && ReturnProperty))
+					{
+						return SNew(SButton)
+							.Text(LOCTEXT("CallFunction", "Call Function"))
+							.OnClicked_Static(&HandleCallFunctionPressed, TWeakObjectPtr<UClass>(ContainerClass), TWeakObjectPtr(Container), TWeakObjectPtr<UFunction>(Function));
+					}
+				}
+			}
+		}
+
+		return SNullWidget::NullWidget;
+	}
+
 	if (const FProperty* LastProperty = Args.Path.GetLastProperty())
 	{
 		const bool bIsVisible = LastProperty->HasAllPropertyFlags(CPF_BlueprintVisible);
@@ -76,10 +287,14 @@ TSharedPtr<SWidget> FPropertyValueFactory::Generate(FGenerateArgs Args) const
 
 			if (const FObjectPropertyBase* ObjectProperty = CastField<const FObjectPropertyBase>(LastProperty))
 			{
-				if (Private::FPropertyViewerMapItem* FoundOnGenerate = Impl->RegisteredFieldEditor.Find(ObjectProperty->PropertyClass->GetStructPathName()))
+				// Class editors should consider parent classes.
+				for (const UClass* Class = ObjectProperty->PropertyClass; !!Class; Class = Class->GetSuperClass())
 				{
-					check(FoundOnGenerate->OnGenerate.IsBound());
-					return FoundOnGenerate->OnGenerate.Execute(Args);
+					if (Private::FPropertyViewerMapItem* FoundOnGenerate = Impl->RegisteredFieldEditor.Find(Class->GetStructPathName()))
+					{
+						check(FoundOnGenerate->OnGenerate.IsBound());
+						return FoundOnGenerate->OnGenerate.Execute(Args);
+					}
 				}
 			}
 			else if (const FStructProperty* StructProperty = CastField<const FStructProperty>(LastProperty))
@@ -108,6 +323,7 @@ TSharedPtr<SWidget> FPropertyValueFactory::Generate(FGenerateArgs Args) const
 			}
 		}
 	}
+
 	return TSharedPtr<SWidget>();
 }
 
@@ -141,6 +357,21 @@ bool FPropertyValueFactory::HasCustomPropertyValue(const FFieldClass* FieldClass
 bool FPropertyValueFactory::HasCustomPropertyValue(const UStruct* Struct) const
 {
 	check(Struct);
+
+	// Class editors should consider parent classes.
+	if (const UClass* Class = Cast<UClass>(Struct))
+	{
+		for (const UClass* Check = Class; !!Check; Check = Check->GetSuperClass())
+		{
+			if (Impl->RegisteredFieldEditor.Find(Check->GetStructPathName()) != nullptr)
+			{
+				return true;
+			}
+		}
+
+		return false;
+	}
+
 	return Impl->RegisteredFieldEditor.Find(Struct->GetStructPathName()) != nullptr;
 }
 
@@ -211,3 +442,5 @@ FPropertyValueFactory::FHandle  FPropertyValueFactory::MakeHandle()
 }
 
 } //namespace
+
+#undef LOCTEXT_NAMESPACE
diff --git a/Engine/Source/Runtime/AdvancedWidgets/Private/Widgets/PropertyViewer/SObjectPropertyValue.cpp b/Engine/Source/Runtime/AdvancedWidgets/Private/Widgets/PropertyViewer/SObjectPropertyValue.cpp
new file mode 100644
--- /dev/null
+++ b/Engine/Source/Runtime/AdvancedWidgets/Private/Widgets/PropertyViewer/SObjectPropertyValue.cpp
@@ -0,0 +1,107 @@
+// Copyright Epic Games, Inc. All Rights Reserved.
+
+#include "Widgets/PropertyViewer/SObjectPropertyValue.h"
+
+#include "Framework/PropertyViewer/INotifyHook.h"
+#include "Styling/SlateTypes.h"
+#include "Styling/AppStyle.h"
+#include "UObject/Object.h"
+#include "UObject/UnrealType.h"
+
+#if WITH_EDITOR
+#include "PropertyCustomizationHelpers.h"
+#include "ThumbnailRendering/ThumbnailManager.h"
+#endif
+
+namespace UE::PropertyViewer
+{
+
+	TSharedPtr<SWidget> SObjectPropertyValue::CreateInstance(const FPropertyValueFactory::FGenerateArgs Args)
+	{
+		return SNew(SObjectPropertyValue)
+			.Path(Args.Path)
+			.NotifyHook(Args.NotifyHook)
+			.IsEnabled(Args.bCanEditValue);
+	}
+
+
+	void SObjectPropertyValue::Construct(const FArguments& InArgs)
+	{
+		Path = InArgs._Path;
+		NotifyHook = InArgs._NotifyHook;
+
+#if WITH_EDITOR
+		if (const FProperty* Property = Path.GetLastProperty())
+		{
+			if (const FObjectPropertyBase* ObjectProperty = CastField<const FObjectPropertyBase>(Property))
+			{
+				if (Property->ArrayDim == 1)
+				{
+					ChildSlot
+					[
+						SNew(SObjectPropertyEntryBox)
+						.ObjectPath(this, &SObjectPropertyValue::GetObjectPath)
+						.ThumbnailPool(UThumbnailManager::Get().GetSharedThumbnailPool())
+						.AllowedClass(ObjectProperty->PropertyClass)
+						.AllowCreate(false)
+						.DisplayCompactSize(true)
+						.DisplayThumbnail(false)
+						.OnObjectChanged(this, &SObjectPropertyValue::OnObjectSelect)
+					];
+				}
+			}
+		}
+#endif
+	}
+
+#if WITH_EDITOR
+	FString SObjectPropertyValue::GetObjectPath() const
+	{
+		if (const void* Container = Path.GetContainerPtr())
+		{
+			if (const FProperty* Property = Path.GetLastProperty())
+			{
+				if (const FObjectPropertyBase* ObjectProperty = CastField<const FObjectPropertyBase>(Property))
+				{
+					TObjectPtr<UObject> Object = nullptr;
+					ObjectProperty->GetValue_InContainer(Container, &Object);
+
+					if (Object)
+					{
+						return Object->GetPathName();
+					}
+
+					return "";
+				}
+			}
+		}
+
+		return "";
+	}
+
+	void SObjectPropertyValue::OnObjectSelect(const FAssetData& InAsset)
+	{
+		if (void* Container = Path.GetContainerPtr())
+		{
+			if (const FProperty* Property = Path.GetLastProperty())
+			{
+				if (const FObjectPropertyBase* ObjectProperty = CastField<const FObjectPropertyBase>(Property))
+				{
+					if (NotifyHook)
+					{
+						NotifyHook->OnPreValueChange(Path);
+					}
+
+					ObjectProperty->SetValue_InContainer(Container, InAsset.GetAsset());
+					
+					if (NotifyHook)
+					{
+						NotifyHook->OnPostValueChange(Path);
+					}
+				}
+			}
+		}
+	}
+#endif
+
+} //namespace
diff --git a/Engine/Source/Runtime/AdvancedWidgets/Private/Widgets/PropertyViewer/SPropertyViewerImpl.cpp b/Engine/Source/Runtime/AdvancedWidgets/Private/Widgets/PropertyViewer/SPropertyViewerImpl.cpp
--- a/Engine/Source/Runtime/AdvancedWidgets/Private/Widgets/PropertyViewer/SPropertyViewerImpl.cpp
+++ b/Engine/Source/Runtime/AdvancedWidgets/Private/Widgets/PropertyViewer/SPropertyViewerImpl.cpp
@@ -903,7 +903,7 @@ TSharedRef<SWidget> FPropertyViewerImpl::CreateTree(bool bHasPreWidget, bool bSh
 			HeaderRowWidget->AddColumn(
 				SHeaderRow::FColumn::FArguments()
 				.ColumnId(ColumnName_PropertyValue)
-				.FillSized(100)
+				.MinSize(150.f)
 				.DefaultLabel(LOCTEXT("PropertyValue", "Field Value"))
 				.FillWidth(0.25f)
 			);
@@ -1105,44 +1105,37 @@ TSharedRef<ITableRow> FPropertyViewerImpl::HandleGenerateRow(TSharedPtr<FTreeNod
 					if (ItemPin->IsField())
 					{
 						FFieldVariant Field = ItemPin->GetField();
-						if (!Field.IsUObject())
+
+						bool bCanEditContainer = false;
+						if (const TSharedPtr<FContainer> OwnerContainer = ItemPin->GetOwnerContainer())
+						{
+							bCanEditContainer = OwnerContainer->CanEdit();
+						}
+
+						FPropertyValueFactory::FGenerateArgs Args;
+						Args.Path = ItemPin->GetPropertyPath();
+						Args.Field = &Field;
+						Args.NotifyHook = PropertyViewOwnerPin->NotifyHook;
+						Args.bCanEditValue = bCanEditContainer
+							&& PropertyViewOwnerPin->PropertyVisibility == SPropertyViewer::EPropertyVisibility::Editable
+							&& Args.Path.GetLastProperty() != nullptr
+							&& !Args.Path.GetLastProperty()->HasAllPropertyFlags(CPF_BlueprintReadOnly);
+
+						FAdvancedWidgetsModule& Module = FAdvancedWidgetsModule::GetModule();
+						TSharedPtr<SWidget> ValueWidget = Module.GetPropertyValueFactory().Generate(Args);
+						if (!ValueWidget)
 						{
-							bool bCanEditContainer = false;
-							if (const TSharedPtr<FContainer> OwnerContainer = ItemPin->GetOwnerContainer())
-							{
-								bCanEditContainer = OwnerContainer->CanEdit();
-							}
-
-							FPropertyValueFactory::FGenerateArgs Args;
-							Args.Path = ItemPin->GetPropertyPath();
-							Args.NotifyHook = PropertyViewOwnerPin->NotifyHook;
-							Args.bCanEditValue = bCanEditContainer
-								&& PropertyViewOwnerPin->PropertyVisibility == SPropertyViewer::EPropertyVisibility::Editable
-								&& Args.Path.GetLastProperty() != nullptr
-								&& !Args.Path.GetLastProperty()->HasAllPropertyFlags(CPF_BlueprintReadOnly);
-
-							FAdvancedWidgetsModule& Module = FAdvancedWidgetsModule::GetModule();
-							TSharedPtr<SWidget> ValueWidget = Module.GetPropertyValueFactory().Generate(Args);
-							if (!ValueWidget)
-							{
-								ValueWidget = Module.GetPropertyValueFactory().GenerateDefault(Args);
-							}
-
-							if (ValueWidget)
-							{
-								return SNew(SHorizontalBox)
-								+ SHorizontalBox::Slot()
-								.AutoWidth()
-								.VAlign(EVerticalAlignment::VAlign_Center)
+							ValueWidget = Module.GetPropertyValueFactory().GenerateDefault(Args);
+						}
+
+						if (ValueWidget)
+						{
+							return SNew(SBox)
+								.HAlign(HAlign_Left)
+								.VAlign(VAlign_Center)
 								[
 									ValueWidget.ToSharedRef()
-								]
-								+ SHorizontalBox::Slot()
-								.FillWidth(1.f)
-								[
-									SNullWidget::NullWidget
 								];
-							}
 						}
 					}
 				}
diff --git a/Engine/Source/Runtime/AdvancedWidgets/Private/Widgets/PropertyViewer/SStringPropertyValue.cpp b/Engine/Source/Runtime/AdvancedWidgets/Private/Widgets/PropertyViewer/SStringPropertyValue.cpp
--- a/Engine/Source/Runtime/AdvancedWidgets/Private/Widgets/PropertyViewer/SStringPropertyValue.cpp
+++ b/Engine/Source/Runtime/AdvancedWidgets/Private/Widgets/PropertyViewer/SStringPropertyValue.cpp
@@ -8,17 +8,24 @@
 #include "UObject/TextProperty.h"
 #include "UObject/UnrealType.h"
 
-#include "Widgets/Input/SEditableText.h"
+#include "Widgets/Input/SEditableTextBox.h"
+#include "Widgets/Layout/SBox.h"
 
 namespace UE::PropertyViewer
 {
 
 TSharedPtr<SWidget> SStringPropertyValue::CreateInstance(const FPropertyValueFactory::FGenerateArgs Args)
 {
-	return SNew(SStringPropertyValue)
-		.Path(Args.Path)
-		.NotifyHook(Args.NotifyHook)
-		.IsEnabled(Args.bCanEditValue);
+	return SNew(SBox)
+		.HAlign(HAlign_Fill)
+		.VAlign(VAlign_Center)
+		.MinDesiredWidth(100.f)
+		[
+			SNew(SStringPropertyValue)
+			.Path(Args.Path)
+			.NotifyHook(Args.NotifyHook)
+			.IsEnabled(Args.bCanEditValue)
+		];
 }
 
 
@@ -34,7 +41,7 @@ void SStringPropertyValue::Construct(const FArguments& InArgs)
 		{
 			ChildSlot
 			[
-				SNew(SEditableText)
+				SNew(SEditableTextBox)
 				.SelectAllTextWhenFocused(true)
 				.Text(this, &SStringPropertyValue::GetText)
 				.OnTextCommitted(this, &SStringPropertyValue::OnTextCommitted)
@@ -87,7 +94,7 @@ void SStringPropertyValue::OnTextCommitted(const FText& InText, ETextCommit::Typ
 			{
 				NotifyHook->OnPostValueChange(Path);
 			}
-			}
+		}
 		else if (const FTextProperty* TextProperty = CastField<const FTextProperty>(Path.GetLastProperty()))
 		{
 			if (NotifyHook)
diff --git a/Engine/Source/Runtime/AdvancedWidgets/Public/Framework/PropertyViewer/PropertyValueFactory.h b/Engine/Source/Runtime/AdvancedWidgets/Public/Framework/PropertyViewer/PropertyValueFactory.h
--- a/Engine/Source/Runtime/AdvancedWidgets/Public/Framework/PropertyViewer/PropertyValueFactory.h
+++ b/Engine/Source/Runtime/AdvancedWidgets/Public/Framework/PropertyViewer/PropertyValueFactory.h
@@ -50,6 +50,7 @@ public:
 	struct FGenerateArgs
 	{
 		FPropertyPath Path;
+		FFieldVariant* Field = nullptr;
 		INotifyHook* NotifyHook = nullptr;
 		bool bCanEditValue = false;
 	};
diff --git a/Engine/Source/Runtime/AdvancedWidgets/Public/Widgets/PropertyViewer/SObjectPropertyValue.h b/Engine/Source/Runtime/AdvancedWidgets/Public/Widgets/PropertyViewer/SObjectPropertyValue.h
new file mode 100644
--- /dev/null
+++ b/Engine/Source/Runtime/AdvancedWidgets/Public/Widgets/PropertyViewer/SObjectPropertyValue.h
@@ -0,0 +1,39 @@
+// Copyright Epic Games, Inc. All Rights Reserved.
+
+#pragma once
+
+#include "Widgets/SCompoundWidget.h"
+#include "Framework/PropertyViewer/PropertyPath.h"
+#include "Framework/PropertyViewer/PropertyValueFactory.h"
+
+class INotifyHook;
+
+namespace UE::PropertyViewer
+{
+
+	/** */
+	class SObjectPropertyValue : public SCompoundWidget
+	{
+	public:
+		static ADVANCEDWIDGETS_API TSharedPtr<SWidget> CreateInstance(const FPropertyValueFactory::FGenerateArgs Args);
+
+	public:
+		SLATE_BEGIN_ARGS(SObjectPropertyValue) {}
+			SLATE_ARGUMENT(FPropertyPath, Path);
+			SLATE_ARGUMENT(INotifyHook*, NotifyHook);
+		SLATE_END_ARGS()
+
+		ADVANCEDWIDGETS_API void Construct(const FArguments& InArgs);
+
+	private:
+		FPropertyPath Path;
+		INotifyHook* NotifyHook = nullptr;
+
+#if WITH_EDITOR
+		FString GetObjectPath() const;
+
+		void OnObjectSelect(const FAssetData& InAsset);
+#endif
+	};
+
+} //namespace
